Part1 of moving Decoders into new modules.
Integrate copies vanilla into new locations.

diff --git a/Engine/Source/Runtime/AdpcmAudioDecoder/Module/Public/Decoders/ADPCMAudioInfo.h b/Engine/Source/Runtime/AdpcmAudioDecoder/Module/Public/Decoders/ADPCMAudioInfo.h
new file mode 100644
--- /dev/null
+++ b/Engine/Source/Runtime/AdpcmAudioDecoder/Module/Public/Decoders/ADPCMAudioInfo.h
@@ -0,0 +1,189 @@
+// Copyright Epic Games, Inc. All Rights Reserved.
+
+/*=============================================================================
+	ADPCMAudioInfo.h: Unreal audio ADPCM decompression interface object.
+=============================================================================*/
+
+#pragma once
+
+#include "CoreMinimal.h"
+#include "Audio.h"
+#include "AudioDecompress.h"
+#include "Sound/SoundWave.h"
+#include "ContentStreaming.h"
+
+#define NUM_ADAPTATION_TABLE 16
+#define NUM_ADAPTATION_COEFF 7
+
+#define WAVE_FORMAT_LPCM  1
+#ifndef WAVE_FORMAT_ADPCM
+#define WAVE_FORMAT_ADPCM 2
+#endif
+
+#if PLATFORM_SUPPORTS_PRAGMA_PACK
+#pragma pack(push, 2)
+#endif
+	struct WaveFormatHeader
+	{
+		uint16 wFormatTag;      // Format type: 1 = PCM, 2 = ADPCM
+		uint16 nChannels;       // Number of channels (i.e. mono, stereo...).
+		uint32 nSamplesPerSec;  // Sample rate. 44100 or 22050 or 11025  Hz.
+		uint32 nAvgBytesPerSec; // For buffer estimation  = sample rate * BlockAlign.
+		uint16 nBlockAlign;     // Block size of data = Channels times BYTES per sample.
+		uint16 wBitsPerSample;  // Number of bits per sample of mono data.
+		uint16 cbSize;          // The count in bytes of the size of extra information (after cbSize).
+	};
+#if PLATFORM_SUPPORTS_PRAGMA_PACK
+#pragma pack(pop)
+#endif
+
+namespace ADPCM
+{
+
+	template <typename T>
+	static void GetAdaptationCoefficients(T(& OutAdaptationCoefficient1)[NUM_ADAPTATION_COEFF], T(& OutAdaptationCoefficient2)[NUM_ADAPTATION_COEFF])
+	{
+		// Magic values as specified by standard
+		static T AdaptationCoefficient1[] =
+		{
+			256, 512, 0, 192, 240, 460, 392
+		};
+		static T AdaptationCoefficient2[] =
+		{
+			0, -256, 0,  64, 0, -208, -232
+		};
+
+		FMemory::Memcpy(&OutAdaptationCoefficient1, AdaptationCoefficient1, sizeof(AdaptationCoefficient1));
+		FMemory::Memcpy(&OutAdaptationCoefficient2, AdaptationCoefficient2, sizeof(AdaptationCoefficient2));
+	}
+
+#if PLATFORM_SUPPORTS_PRAGMA_PACK
+#pragma pack(push, 2)
+#endif
+	struct ADPCMFormatHeader
+	{
+		WaveFormatHeader	BaseFormat;
+		uint16				wSamplesPerBlock;
+		uint16				wNumCoef;
+		int16				aCoef[2 * NUM_ADAPTATION_COEFF];
+		uint32				SamplesPerChannel;	// This is the exact samples per channel for sample precise looping
+
+		ADPCMFormatHeader()
+		{
+			int16 AdaptationCoefficient1[NUM_ADAPTATION_COEFF];
+			int16 AdaptationCoefficient2[NUM_ADAPTATION_COEFF];
+			GetAdaptationCoefficients(AdaptationCoefficient1, AdaptationCoefficient2);
+
+			// Interlace the coefficients as pairs
+			for (int32 Coeff = 0, WriteIndex = 0; Coeff < NUM_ADAPTATION_COEFF; Coeff++)
+			{
+				aCoef[WriteIndex++] = AdaptationCoefficient1[Coeff];
+				aCoef[WriteIndex++] = AdaptationCoefficient2[Coeff];
+			}
+		}
+	};
+#if PLATFORM_SUPPORTS_PRAGMA_PACK
+#pragma pack(pop)
+#endif
+
+};
+
+class FADPCMAudioInfo : public ICompressedAudioInfo
+{
+public:
+	ENGINE_API FADPCMAudioInfo(void);
+	ENGINE_API virtual ~FADPCMAudioInfo(void);
+
+	// ICompressedAudioInfo Interface
+	ENGINE_API virtual bool ReadCompressedInfo(const uint8* InSrcBufferData, uint32 InSrcBufferDataSize, struct FSoundQualityInfo* QualityInfo);
+	ENGINE_API virtual bool ReadCompressedData(uint8* Destination, bool bLooping, uint32 BufferSize);
+	ENGINE_API virtual void SeekToTime(const float SeekTime);
+	ENGINE_API virtual void SeekToFrame(const uint32 SeekFrame);
+	ENGINE_API virtual void ExpandFile(uint8* DstBuffer, struct FSoundQualityInfo* QualityInfo);
+	virtual void EnableHalfRate(bool HalfRate) {};
+	virtual uint32 GetSourceBufferSize() const { return SrcBufferDataSize; }
+	virtual bool UsesVorbisChannelOrdering() const { return false; }
+	virtual int GetStreamBufferSize() const;
+
+	// Additional overrides for streaming
+	virtual bool SupportsStreaming() const override {return true;}
+	virtual bool StreamCompressedInfoInternal(const FSoundWaveProxyPtr& InWaveProxy, struct FSoundQualityInfo* QualityInfo) override;
+	virtual bool StreamCompressedData(uint8* Destination, bool bLooping, uint32 BufferSize, int32& OutNumBytesStreamed) override;
+	virtual int32 GetCurrentChunkIndex() const override
+	{
+		return CurrentChunkIndex;
+	}
+	virtual int32 GetCurrentChunkOffset() const override {return CurrentChunkBufferOffset;}
+	virtual bool ReleaseStreamChunk(bool bBlockUntilReleased) override;
+	// End of ICompressedAudioInfo Interface
+
+	uint16 GetFormatTag()
+	{
+		if (WaveInfo.pFormatTag)
+		{
+			return *WaveInfo.pFormatTag;
+		}
+		else
+		{
+			return 0;
+		}
+	}
+
+	int32 GetNumChannels()
+	{
+		return NumChannels;
+	}
+
+private:
+
+	// Wrapper function that returns a pointer to the currently used compressed data.
+	// If a non-zero chunk is requested, this function also aquires a reference to that chunk
+	// until we move on to a different chunk.
+	const uint8* GetLoadedChunk(const FSoundWaveProxyPtr& InSoundWave, uint32 ChunkIndex, uint32& OutChunkSize);
+
+	int32 NumConsecutiveReadFailiures;
+
+	FWaveModInfo WaveInfo;
+	const uint8*	SrcBufferData;
+	uint32			SrcBufferDataSize;
+
+	uint32	UncompressedBlockSize;
+	uint32	CompressedBlockSize;
+	uint32	BlockSize;
+	int32	StreamBufferSize;
+	uint32	TotalDecodedSize;
+	int32	NumChannels;
+	int32	Format;
+
+	uint32 PreviouslyRequestedChunkIndex;
+
+	uint8*			UncompressedBlockData;			// This holds the current block of compressed data for all channels
+	uint32			CurrentUncompressedBlockSampleIndex;	// This is the sample index within the current uncompressed block data
+	uint32			CurrentChunkIndex;				// This is the index that is currently being used, needed by streaming engine to make sure it stays loaded and the next chunk gets preloaded
+	uint32			CurrentChunkBufferOffset;		// This is this byte offset within the current chunk, used by streaming engine to prioritize a load if more then half way through current chunk
+	uint32			CurrentChunkDataSize;			// The size of the current chunk, the first chunk is bigger to accommodate the header info
+	uint32			TotalSamplesStreamed;			// The number of samples streamed so far (per channel)
+	uint32			TotalSamplesPerChannel;			// Number of samples per channel, used to detect when an audio waveform has ended
+	uint32			SamplesPerBlock;				// The number of samples per block
+	uint32			FirstChunkSampleDataOffset;		// The size of the header in the first chunk, used to skip over it when looping or starting the sample over
+	uint32          FirstChunkSampleDataIndex;
+	const uint8*	CurCompressedChunkData;			// A pointer to the current chunk of data
+	FAudioChunkHandle CurCompressedChunkHandle;     // Shared reference to the current chunk of data.
+	FCriticalSection CurCompressedChunkHandleCriticalSection;
+	bool bDecompressorReleased;                     // When the chunk is released, we raise this to true, then early exit on future decodes.
+
+	uint32			CurrentCompressedBlockIndex;		// For non disk streaming - the current compressed block in the compressed source data
+	uint32			TotalCompressedBlocksPerChannel;	// For non disk streaming - the total number of compressed blocks per channel
+	uint8			bNewSeekRequest : 1;				// Whether or not there is a new seek request to be processed 
+	uint8			bSeekPendingRead : 1;				// Whether or not a requested seek is pending a read
+	uint8           bSeekedFowardToNextChunk : 1;       // If this is true, we have already seeked forward to the next chunk while waiting for the current chunk of audio to load.
+
+private:
+	void ProcessSeekRequest();
+	void SeekToFrameInternal(const uint32 InSeekFrame);
+	void ResetSeekState();
+
+	std::atomic<uint32> TargetSeekFrame;
+
+	FCriticalSection StreamSeekCriticalSection;
+};
diff --git a/Engine/Source/Runtime/OpusAudioDecoder/Module/Public/Decoders/OpusAudioInfo.h b/Engine/Source/Runtime/OpusAudioDecoder/Module/Public/Decoders/OpusAudioInfo.h
new file mode 100644
--- /dev/null
+++ b/Engine/Source/Runtime/OpusAudioDecoder/Module/Public/Decoders/OpusAudioInfo.h
@@ -0,0 +1,92 @@
+// Copyright Epic Games, Inc. All Rights Reserved.
+
+/*=============================================================================
+	OpusAudioInfo.h: Unreal audio opus decompression interface object.
+=============================================================================*/
+
+#pragma once
+
+#include "CoreMinimal.h"
+#include "AudioDecompress.h"
+
+struct FOpusDecoderWrapper;
+struct FSoundQualityInfo;
+
+/**
+* Helper class to parse opus data
+*/
+class FOpusAudioInfo : public IStreamedCompressedInfo
+{
+public:
+	ENGINE_API FOpusAudioInfo();
+	ENGINE_API virtual ~FOpusAudioInfo();
+
+	//~ Begin IStreamedCompressedInfo Interface
+	bool ParseHeader(const uint8* InSrcBufferData, uint32 InSrcBufferDataSize, FSoundQualityInfo* QualityInfo) override;
+	int32 GetFrameSize() override;
+	uint32 GetMaxFrameSizeSamples() const override;
+	bool CreateDecoder() override;
+	FDecodeResult Decode(const uint8* CompressedData, const int32 CompressedDataSize, uint8* OutPCMData, const int32 OutputPCMDataSize) override;
+	void PrepareToLoop() override;
+	void SeekToTime(const float SeekTime) override;
+	void SeekToFrame(const uint32 SeekFrame) override;
+	//~ End IStreamedCompressedInfo Interface
+
+	struct ENGINE_API FHeader
+	{
+		char Identifier[8];
+		uint8 Version = 0;
+		uint8 NumChannels = 0;
+		uint32 SampleRate = 0;
+		uint32 EncodedSampleRate = 0;
+		uint64 ActiveSampleCount = 0;
+		uint32 NumEncodedFrames = 0;
+		int32 NumPreSkipSamples = 0;
+		int32 NumSilentSamplesAtBeginning = 0;
+		int32 NumSilentSamplesAtEnd = 0;
+
+		void Reset()
+		{
+			FMemory::Memzero(Identifier);
+			Version = 0;
+			NumChannels = 0;
+			SampleRate = 0;
+			EncodedSampleRate = 0;
+			ActiveSampleCount = 0;
+			NumEncodedFrames = 0;
+			NumPreSkipSamples = 0;
+			NumSilentSamplesAtBeginning = 0;
+			NumSilentSamplesAtEnd = 0;
+		}
+
+		static inline const char OPUS_ID[8] {'U','E','O','P','U','S','\0','\0'};
+		static constexpr int32 HeaderSize()
+		{
+			return
+				  sizeof(char)*8	// Identifier
+				+ sizeof(uint8)		// Version
+				+ sizeof(uint8)		// NumChannels
+				+ sizeof(uint32)	// SampleRate
+				+ sizeof(uint32)	// EncodedSampleRate
+				+ sizeof(uint64)	// ActiveSampleCount
+				+ sizeof(uint32)	// NumEncodedFrames
+				+ sizeof(int32)		// NumPreSkipSamples
+				+ sizeof(int32)		// NumSilentSamplesAtBeginning
+				+ sizeof(int32);	// NumSilentSamplesAtEnd
+		}
+	};
+	static ENGINE_API bool ParseHeader(FHeader& OutHeader, uint32& OutNumRead, const uint8* InSrcBufferData, uint32 InSrcBufferDataSize);
+
+protected:
+	/** Wrapper around Opus-specific decoding state and APIs */
+	FOpusDecoderWrapper* OpusDecoderWrapper;
+
+	/** The currently active header. */
+	FHeader Header;
+	
+	/** The number of decoded samples that need to be skipped over. */
+	int32 NumRemainingSamplesToSkip = 0;
+
+	/** Samples that are left over from the previous Decode() call that did not fit into the receive buffer. */
+	TArray<uint8> PreviousDecodedUnusedSamples;
+};
diff --git a/Engine/Source/Runtime/VorbisAudioDecoder/Module/Public/Decoders/VorbisAudioInfo.h b/Engine/Source/Runtime/VorbisAudioDecoder/Module/Public/Decoders/VorbisAudioInfo.h
new file mode 100644
--- /dev/null
+++ b/Engine/Source/Runtime/VorbisAudioDecoder/Module/Public/Decoders/VorbisAudioInfo.h
@@ -0,0 +1,144 @@
+// Copyright Epic Games, Inc. All Rights Reserved.
+
+/*=============================================================================
+	VorbisAudioInfo.h: Unreal audio vorbis decompression interface object.
+=============================================================================*/
+
+#pragma once
+
+#include "CoreMinimal.h"
+#include "HAL/ThreadSafeBool.h"
+#include "AudioDecompress.h"
+
+/**
+ * Whether to use OggVorbis audio format.
+ **/
+#ifndef WITH_OGGVORBIS
+	#define WITH_OGGVORBIS 0
+#endif
+
+/**
+ * Whether to load OggVorbis DLLs
+ **/
+#ifndef WITH_OGGVORBIS_DLL
+	#define WITH_OGGVORBIS_DLL (PLATFORM_WINDOWS || PLATFORM_HOLOLENS)
+#endif
+
+
+namespace VorbisChannelInfo
+{
+	extern ENGINE_API const int32 Order[8][8];
+}
+
+/**
+ * Loads vorbis dlls
+*/
+ENGINE_API void LoadVorbisLibraries();
+
+#if WITH_OGGVORBIS
+
+/** 
+ * Helper class to parse ogg vorbis data
+ */
+class FVorbisAudioInfo : public ICompressedAudioInfo
+{
+public:
+	ENGINE_API FVorbisAudioInfo( void );
+	ENGINE_API virtual ~FVorbisAudioInfo( void );
+
+	/** Emulate read from memory functionality */
+	size_t			ReadMemory( void *ptr, uint32 size );
+	int				SeekMemory( uint32 offset, int whence );
+	int				CloseMemory( void );
+	long			TellMemory( void );
+
+	/** Emulate read from streaming functionality */
+	size_t			ReadStreaming( void *ptr, uint32 size );
+	int				CloseStreaming( void );
+
+	/** Common info and data functions between ReadCompressedInfo/ReadCompressedData and StreamCompressedInfo/StreamCompressedData */
+	/** Can't use ov_callbacks here so we have to pass in a void* for the callbacks instead */
+	bool GetCompressedInfoCommon(void*	Callbacks, FSoundQualityInfo* QualityInfo);
+
+	/** 
+	 * Reads the header information of an ogg vorbis file
+	 * 
+	 * @param	Resource		Info about vorbis data
+	 */
+	ENGINE_API virtual bool ReadCompressedInfo( const uint8* InSrcBufferData, uint32 InSrcBufferDataSize, struct FSoundQualityInfo* QualityInfo ) override;
+
+	/** 
+	 * Decompresses ogg data to raw PCM data. 
+	 * 
+	 * @param	Destination	where to place the decompressed sound
+	 * @param	bLooping	whether to loop the sound by seeking to the start, or pad the buffer with zeroes
+	 * @param	BufferSize	number of bytes of PCM data to create
+	 *
+	 * @return	bool		true if the end of the data was reached (for both single shot and looping sounds)
+	 */
+	ENGINE_API virtual bool ReadCompressedData( uint8* InDestination, bool bLooping, uint32 BufferSize ) override;
+
+	ENGINE_API virtual void SeekToTime( const float SeekTime ) override;
+
+	ENGINE_API virtual void SeekToFrame(const uint32 SeekFrame) override;
+
+	/** 
+	 * Decompress an entire ogg data file to a TArray
+	 */
+	ENGINE_API virtual void ExpandFile( uint8* DstBuffer, struct FSoundQualityInfo* QualityInfo ) override;
+
+	/** 
+	 * Sets ogg to decode to half-rate
+	 * 
+	 * @param	Resource		Info about vorbis data
+	 */
+	ENGINE_API virtual void EnableHalfRate( bool HalfRate ) override;
+
+	virtual uint32 GetSourceBufferSize() const override { return SrcBufferDataSize; }
+
+	virtual bool UsesVorbisChannelOrdering() const override { return true; }
+	virtual int GetStreamBufferSize() const override { return MONO_PCM_BUFFER_SIZE; }
+
+	// Additional overrides for streaming
+	virtual bool SupportsStreaming() const override {return true;}
+	virtual bool StreamCompressedData(uint8* InDestination, bool bLooping, uint32 BufferSize, int32& OutNumBytesStreamed) override;
+	virtual int32 GetCurrentChunkIndex() const override {return CurrentStreamingChunkIndex;}
+	virtual int32 GetCurrentChunkOffset() const override {return BufferOffset % CurrentStreamingChunksSize;}
+	virtual bool HasError() const override;
+	// End of ICompressedAudioInfo Interface
+
+protected:
+	virtual bool StreamCompressedInfoInternal(const FSoundWaveProxyPtr& InWaveProxy, struct FSoundQualityInfo* QualityInfo) override;
+
+	friend class FAudioFormatOgg;
+	ENGINE_API int32 GetAudioDataStartOffset() const;
+
+private:
+	using Super = ICompressedAudioInfo; 
+	const uint8* GetLoadedChunk(FSoundWaveProxyPtr InSoundWave, uint32 ChunkIndex, uint32& OutChunkSize);
+
+	struct FVorbisFileWrapper* VFWrapper;
+	const uint8* SrcBufferData;
+	uint32 SrcBufferDataSize;
+	uint32 BufferOffset;
+	uint32 CurrentBufferChunkOffset;
+
+	// In case Ogg Vorbis fails to return any compressed audio for an asset,
+	// we use this counter to exit out of the decoder's while loop early.
+	int32 TimesLoopedWithoutDecompressedAudio;
+
+	/** Critical section used to prevent multiple threads accessing same ogg-vorbis file handles at the same time */
+	mutable FCriticalSection VorbisCriticalSection;
+
+	uint8 const* CurrentStreamingChunkData;
+	int32 CurrentStreamingChunkIndex;
+	int32 NextStreamingChunkIndex;
+	uint32 CurrentStreamingChunksSize;
+	bool bHeaderParsed;
+	bool bHasError = false;
+	int32 LastErrorCode = 0;
+
+	// This handle is used to ensure that a chunk currently being decoded isn't evicted until we are done with it.
+	FAudioChunkHandle CurCompressedChunkHandle;
+};
+#endif
