Move all Engine Decoders into their own modules
#tests compilation tests, all platforms

diff --git a/Engine/Config/BaseEngine.ini b/Engine/Config/BaseEngine.ini
--- a/Engine/Config/BaseEngine.ini
+++ b/Engine/Config/BaseEngine.ini
@@ -1553,6 +1553,9 @@ DeferStartupPrecache=false
 +FormatModuleHints=AudioFormatOGG
 +FormatModuleHints=AudioFormatADPCM
 +FormatModuleHints=AudioFormatBINK
++AudioInfoModules=OpusAudioDecoder
++AudioInfoModules=VorbisAudioDecoder
++AudioInfoModules=AdpcmAudioDecoder
 +AudioInfoModules=BinkAudioDecoder
 FallbackFormat=BINKA
 PlatformFormat=BINKA
diff --git a/Engine/Source/Developer/AudioFormatADPCM/AudioFormatADPCM.Build.cs b/Engine/Source/Developer/AudioFormatADPCM/AudioFormatADPCM.Build.cs
--- a/Engine/Source/Developer/AudioFormatADPCM/AudioFormatADPCM.Build.cs
+++ b/Engine/Source/Developer/AudioFormatADPCM/AudioFormatADPCM.Build.cs
@@ -7,10 +7,12 @@ public class AudioFormatADPCM : ModuleRules
 	public AudioFormatADPCM(ReadOnlyTargetRules Target) : base(Target)
 	{
 		PrivateIncludePathModuleNames.Add("TargetPlatform");
+		PrivateDependencyModuleNames.Add("AdpcmAudioDecoder");
 		PrivateDependencyModuleNames.Add("Core");
 		if (Target.bCompileAgainstEngine)
 		{
 			PrivateDependencyModuleNames.Add("Engine");
+			
 		}
 	}
 }
diff --git a/Engine/Source/Developer/AudioFormatADPCM/Private/AudioFormatADPCM.cpp b/Engine/Source/Developer/AudioFormatADPCM/Private/AudioFormatADPCM.cpp
--- a/Engine/Source/Developer/AudioFormatADPCM/Private/AudioFormatADPCM.cpp
+++ b/Engine/Source/Developer/AudioFormatADPCM/Private/AudioFormatADPCM.cpp
@@ -12,7 +12,7 @@
 #include "Interfaces/IAudioFormatModule.h"
 #include "AudioDecompress.h"
 #include "Audio.h"
-#include "ADPCMAudioInfo.h"
+#include "Decoders/ADPCMAudioInfo.h"
 
 DEFINE_LOG_CATEGORY_STATIC(LogAudioFormatADPCM, Log, All);
 
diff --git a/Engine/Source/Developer/AudioFormatOgg/AudioFormatOgg.Build.cs b/Engine/Source/Developer/AudioFormatOgg/AudioFormatOgg.Build.cs
--- a/Engine/Source/Developer/AudioFormatOgg/AudioFormatOgg.Build.cs
+++ b/Engine/Source/Developer/AudioFormatOgg/AudioFormatOgg.Build.cs
@@ -18,7 +18,8 @@ public class AudioFormatOgg : ModuleRules
 		PrivateDependencyModuleNames.AddRange(
 			new string[] {
 				"Core",
-				"Engine"
+				"Engine",
+				"VorbisAudioDecoder"
 			}
 		);
 
diff --git a/Engine/Source/Developer/AudioFormatOgg/Private/AudioFormatOgg.cpp b/Engine/Source/Developer/AudioFormatOgg/Private/AudioFormatOgg.cpp
--- a/Engine/Source/Developer/AudioFormatOgg/Private/AudioFormatOgg.cpp
+++ b/Engine/Source/Developer/AudioFormatOgg/Private/AudioFormatOgg.cpp
@@ -5,8 +5,7 @@
 #include "Modules/ModuleManager.h"
 #include "Interfaces/IAudioFormat.h"
 #include "Interfaces/IAudioFormatModule.h"
-#include "VorbisAudioInfo.h"
-
+#include "Decoders/VorbisAudioInfo.h"
 
 #if WITH_OGGVORBIS
 	#pragma pack(push, 8)
diff --git a/Engine/Source/Developer/AudioFormatOpus/AudioFormatOpus.Build.cs b/Engine/Source/Developer/AudioFormatOpus/AudioFormatOpus.Build.cs
--- a/Engine/Source/Developer/AudioFormatOpus/AudioFormatOpus.Build.cs
+++ b/Engine/Source/Developer/AudioFormatOpus/AudioFormatOpus.Build.cs
@@ -18,7 +18,9 @@ public class AudioFormatOpus : ModuleRules
 		PrivateDependencyModuleNames.AddRange(
 			new string[] {
 				"Core",
-				"Engine"
+				"Engine",
+				"VorbisAudioDecoder",	// for VorbisChannelInfo
+				"OpusAudioDecoder"
 			}
 		);
 
diff --git a/Engine/Source/Developer/AudioFormatOpus/Private/AudioFormatOpus.cpp b/Engine/Source/Developer/AudioFormatOpus/Private/AudioFormatOpus.cpp
--- a/Engine/Source/Developer/AudioFormatOpus/Private/AudioFormatOpus.cpp
+++ b/Engine/Source/Developer/AudioFormatOpus/Private/AudioFormatOpus.cpp
@@ -6,8 +6,9 @@
 #include "Modules/ModuleManager.h"
 #include "Interfaces/IAudioFormat.h"
 #include "Interfaces/IAudioFormatModule.h"
-#include "OpusAudioInfo.h"
-#include "VorbisAudioInfo.h"	// for VorbisChannelInfo
+
+#include "Decoders/OpusAudioInfo.h"
+#include "Decoders/VorbisAudioInfo.h"	// for VorbisChannelInfo
 
 THIRD_PARTY_INCLUDES_START
 #include "opus_multistream.h"
diff --git a/Engine/Source/Runtime/AdpcmAudioDecoder/Module/AdpcmAudioDecoder.Build.cs b/Engine/Source/Runtime/AdpcmAudioDecoder/Module/AdpcmAudioDecoder.Build.cs
new file mode 100644
--- /dev/null
+++ b/Engine/Source/Runtime/AdpcmAudioDecoder/Module/AdpcmAudioDecoder.Build.cs
@@ -0,0 +1,20 @@
+// Copyright Epic Games, Inc. All Rights Reserved.
+
+using UnrealBuildTool;
+
+public class AdpcmAudioDecoder : ModuleRules
+{
+    public AdpcmAudioDecoder(ReadOnlyTargetRules Target) : base(Target)
+    {
+        ShortName = "AdpcmAudioDecoder";
+
+        PublicDependencyModuleNames.AddRange(
+            new string[]
+            {
+                "Core",
+                "CoreUObject",
+                "Engine"
+            }
+        );
+    }
+}
diff --git a/Engine/Source/Runtime/AdpcmAudioDecoder/Module/Private/ADPCMAudioInfo.cpp b/Engine/Source/Runtime/AdpcmAudioDecoder/Module/Private/ADPCMAudioInfo.cpp
new file mode 100644
--- /dev/null
+++ b/Engine/Source/Runtime/AdpcmAudioDecoder/Module/Private/ADPCMAudioInfo.cpp
@@ -0,0 +1,1314 @@
+// Copyright Epic Games, Inc. All Rights Reserved.
+
+
+#include "Decoders/ADPCMAudioInfo.h"
+#include "Interfaces/IAudioFormat.h"
+
+static int32 bDisableADPCMSeekingCVar = 0;
+FAutoConsoleVariableRef CVarDisableADPCMSeeking(
+	TEXT("au.adpcm.DisableSeeking"),
+	bDisableADPCMSeekingCVar,
+	TEXT("Disables seeking with ADPCM.\n"),
+	ECVF_Default);
+
+static int32 ADPCMReadFailiureTimeoutCVar = 64;
+FAutoConsoleVariableRef CVarADPCMReadFailiureTimeout(
+	TEXT("au.adpcm.ADPCMReadFailiureTimeout"),
+	ADPCMReadFailiureTimeoutCVar,
+	TEXT("Sets the number of ADPCM decode attempts we'll try before stopping the sound wave altogether.\n"),
+	ECVF_Default);
+
+static int32 ADPCMDisableSeekForwardOnChunkMissesCVar = 1;
+FAutoConsoleVariableRef CVarADPCMDisableSeekForwardOnChunkMisses(
+	TEXT("au.adpcm.DisableSeekForwardOnReadMisses"),
+	ADPCMDisableSeekForwardOnChunkMissesCVar,
+	TEXT("When there is a seek pending and this CVar is set to 0, we will scan forward in the file.\n"),
+	ECVF_Default);
+
+static int32 ADPCMOnlySeekForwardOneChunkCVar = 1;
+FAutoConsoleVariableRef CVarADPCMOnlySeekForwardOneChunk(
+	TEXT("au.adpcm.OnlySeekForwardOneChunk"),
+	ADPCMOnlySeekForwardOneChunkCVar,
+	TEXT("When set to 1, we will not continue to seek forward after failing to load two chunks in a row.\n"),
+	ECVF_Default);
+
+static float ChanceForIntentionalChunkMissCVar = 0.0f;
+FAutoConsoleVariableRef CVarChanceForIntentionalChunkMiss(
+	TEXT("au.adpcm.ChanceForIntentionalChunkMiss"),
+	ChanceForIntentionalChunkMissCVar,
+	TEXT("If this is set > 0 we will intentionally drop chunks. Used for debugging..\n"),
+	ECVF_Default);
+
+#define WAVE_FORMAT_LPCM  1
+#define WAVE_FORMAT_ADPCM 2
+
+namespace ADPCM
+{
+	void DecodeBlock(const uint8* EncodedADPCMBlock, int32 BlockSize, int16* DecodedPCMData);
+	void DecodeBlockStereo(const uint8* EncodedADPCMBlockLeft, const uint8* EncodedADPCMBlockRight, int32 BlockSize, int16* DecodedPCMData);
+}
+
+FADPCMAudioInfo::FADPCMAudioInfo(void)
+	: NumConsecutiveReadFailiures(0)
+	, UncompressedBlockSize(0)
+	, CompressedBlockSize(0)
+	, BlockSize(0)
+	, StreamBufferSize(0)
+	, TotalDecodedSize(0)
+	, NumChannels(0)
+	, Format(0)
+	, PreviouslyRequestedChunkIndex(0)
+	, UncompressedBlockData(nullptr)
+	, SamplesPerBlock(0)
+	, FirstChunkSampleDataOffset(0)
+	, FirstChunkSampleDataIndex(0)
+	, bDecompressorReleased(false)
+	, bNewSeekRequest(false)
+	, bSeekPendingRead(false)
+	, bSeekedFowardToNextChunk(false)
+	, TargetSeekFrame(0)
+{
+}
+
+FADPCMAudioInfo::~FADPCMAudioInfo(void)
+{
+	if(UncompressedBlockData != nullptr)
+	{
+		FMemory::Free(UncompressedBlockData);
+		UncompressedBlockData = nullptr;
+	}
+}
+
+void FADPCMAudioInfo::SeekToTime(const float InSeekTime)
+{
+	if (bDisableADPCMSeekingCVar)
+	{
+		return;
+	}
+
+	const uint32 SamplesPerSec = *WaveInfo.pSamplesPerSec;
+	TargetSeekFrame = static_cast<uint32>(InSeekTime * SamplesPerSec);
+	bNewSeekRequest = true;
+}
+
+void FADPCMAudioInfo::SeekToFrame(const uint32 InSeekFrame)
+{
+	if (bDisableADPCMSeekingCVar)
+	{
+		return;
+	}
+
+	TargetSeekFrame = InSeekFrame;
+	bNewSeekRequest = true;
+}
+
+void FADPCMAudioInfo::SeekToFrameInternal(const uint32 InSeekFrame)
+{
+	// Reset chunk handle in preperation for a new chunk.
+	CurCompressedChunkData = nullptr;
+
+	// zero out any existing decode data as it will stale after seeking
+	if (UncompressedBlockData)
+	{
+		FMemory::Memzero(UncompressedBlockData, NumChannels * UncompressedBlockSize);
+	}
+	
+	float SeekTimeSeconds = InSeekFrame / static_cast<float>(*WaveInfo.pSamplesPerSec);
+	UE_LOG(LogAudio, Verbose, TEXT("Seeking ADPCM source to %.3f sec"), SeekTimeSeconds);
+
+	if (InSeekFrame == 0)
+	{
+		CurrentCompressedBlockIndex = 0;
+		// This is set to the max value to trigger the decompression of the first audio block
+		CurrentUncompressedBlockSampleIndex = UncompressedBlockSize / sizeof(uint16);;
+		CurrentChunkIndex = FirstChunkSampleDataIndex;
+		CurrentChunkBufferOffset = FirstChunkSampleDataOffset;
+		TotalSamplesStreamed = 0;
+		
+		ResetSeekState();
+		return;
+	}
+
+	// Calculate block index & force SeekTime to be in bounds.
+	check(WaveInfo.pSamplesPerSec != nullptr);
+	uint32 SeekedSamples = InSeekFrame;
+	TotalSamplesStreamed = FMath::Min<uint32>(SeekedSamples, TotalSamplesPerChannel - 1);
+
+	const uint32 HeaderOffset = static_cast<uint32>(WaveInfo.SampleDataStart - SrcBufferData);
+	
+	if (!StreamingSoundWave)
+	{
+		// For the non streaming case:
+		if (Format == WAVE_FORMAT_LPCM)
+		{
+			// There are no "blocks" on LPCM, so only update the total samples streamed (which is based off sample rate).
+			// Note that TotalSamplesStreamed is per-channel in the ReadCompressedInfo. Channels are taken into account there.
+			TotalSamplesStreamed = FMath::Clamp<uint32>(SeekedSamples, 0, TotalSamplesPerChannel - 1);
+		}
+		else
+		{
+			// Clamp to the end of memory in case we have an invalid seek time.
+			SeekedSamples = FMath::Clamp<uint32>(SeekedSamples, 0, TotalSamplesPerChannel - 1);
+
+			// Compute the block index that we're seeked to
+			CurrentCompressedBlockIndex = SeekedSamples / SamplesPerBlock;
+
+			// Update the samples streamed to the current block index and the samples per block
+			TotalSamplesStreamed = CurrentCompressedBlockIndex * SamplesPerBlock;
+		}
+	}
+	else
+	{
+		if (StreamingSoundWave->GetNumChunks() == 0)
+		{
+			UE_LOG(LogAudio, Error, TEXT("Entered streaming seek path with a non-streaming sound!"));
+			return;
+		}
+
+		const uint32 TotalStreamingChunks = StreamingSoundWave->GetNumChunks();
+
+		if (Format == WAVE_FORMAT_ADPCM)
+		{
+			CurrentCompressedBlockIndex = TotalSamplesStreamed / SamplesPerBlock; // Compute the block index that where SeekTime resides.
+			CurrentChunkIndex = FirstChunkSampleDataIndex;
+			CurrentChunkBufferOffset = FirstChunkSampleDataOffset;
+
+			const int32 ChannelBlockSize = BlockSize * NumChannels;
+			for (uint32 BlockIndex = 0; BlockIndex < CurrentCompressedBlockIndex; ++BlockIndex)
+			{
+				const uint32 SizeOfChunk = StreamingSoundWave->GetSizeOfChunk(CurrentChunkIndex);
+				if (CurrentChunkBufferOffset + ChannelBlockSize >= SizeOfChunk)
+				{
+					const uint32 RemainderAfterEndOfChunk = CurrentChunkBufferOffset + ChannelBlockSize - SizeOfChunk;
+					ensureMsgf(RemainderAfterEndOfChunk == 0, TEXT("Found partial ADPCM block of %u samples- Please check FADPCMAudioFormat::SplitDataForStreaming for errors."), RemainderAfterEndOfChunk);
+					++CurrentChunkIndex;
+					CurrentChunkBufferOffset = 0;
+				}
+				else
+				{
+					// Always add chunks in NumChannels pairs
+					CurrentChunkBufferOffset += ChannelBlockSize;
+				}
+				
+				if (CurrentChunkIndex >= TotalStreamingChunks)
+				{
+					CurrentChunkIndex = FirstChunkSampleDataIndex;
+					CurrentChunkBufferOffset = FirstChunkSampleDataOffset;
+					break;
+				}
+			}
+		}
+		else if (Format == WAVE_FORMAT_LPCM)
+		{
+			const int32 ChannelBlockSize = sizeof(int16) * NumChannels;
+			const uint32 TotalByteSizeToSeek = TotalSamplesStreamed * ChannelBlockSize;	
+			
+			uint32 TargetChunkIndex = FirstChunkSampleDataIndex;
+			uint32 SizeOfTargetChunk = StreamingSoundWave->GetSizeOfChunk(TargetChunkIndex);
+			
+			uint32 CurrentBytes = 0;
+			while (CurrentBytes + SizeOfTargetChunk < TotalByteSizeToSeek)
+			{
+				const bool TargetChunkInRange = TargetChunkIndex < TotalStreamingChunks;
+				ensure(TargetChunkInRange);
+
+				if (!TargetChunkInRange)
+				{
+					break;
+				}
+
+				++TargetChunkIndex;
+				CurrentBytes += SizeOfTargetChunk;
+				CurrentChunkBufferOffset -= SizeOfTargetChunk;
+				SizeOfTargetChunk = StreamingSoundWave->GetSizeOfChunk(TargetChunkIndex);
+			}
+
+			check(CurrentBytes <= TotalByteSizeToSeek);
+
+			CurrentChunkIndex = TargetChunkIndex;
+			CurrentChunkBufferOffset = TotalByteSizeToSeek - CurrentBytes;
+			CurrentChunkBufferOffset -= CurrentChunkBufferOffset % ChannelBlockSize;
+		}
+		else
+		{
+			// If we hit this, Format was invalid:
+			checkNoEntry();
+			return;
+		}
+	}
+	bSeekPendingRead = true;
+}
+
+bool FADPCMAudioInfo::ReadCompressedInfo(const uint8* InSrcBufferData, uint32 InSrcBufferDataSize, struct FSoundQualityInfo* QualityInfo)
+{
+	if (!InSrcBufferData)
+	{
+		FString Name = QualityInfo ? QualityInfo->DebugName : TEXT("Unknown");
+		UE_LOG(LogAudio, Warning, TEXT("Failed to read compressed ADPCM audio from ('%s') because there was no resource data."), *Name);
+
+		return false;
+	}
+
+	SrcBufferData = InSrcBufferData;
+	SrcBufferDataSize = InSrcBufferDataSize;
+
+	void*	FormatHeader;
+
+	if (!WaveInfo.ReadWaveInfo((uint8*)SrcBufferData, SrcBufferDataSize, nullptr, false, &FormatHeader))
+	{
+		UE_LOG(LogAudio, Warning, TEXT("WaveInfo.ReadWaveInfo Failed"));
+		return false;
+	}
+
+	Format = *WaveInfo.pFormatTag;
+	NumChannels = *WaveInfo.pChannels;
+
+	if (Format == WAVE_FORMAT_ADPCM)
+	{
+		ADPCM::ADPCMFormatHeader* ADPCMHeader = (ADPCM::ADPCMFormatHeader*)FormatHeader;
+		TotalSamplesPerChannel = ADPCMHeader->SamplesPerChannel;
+		SamplesPerBlock = ADPCMHeader->wSamplesPerBlock;
+
+		const uint32 PreambleSize = 7;
+		BlockSize = *WaveInfo.pBlockAlign;
+
+		// ADPCM starts with 2 uncompressed samples and then the remaining compressed sample data has 2 samples per byte
+		UncompressedBlockSize = (2 + (BlockSize - PreambleSize) * 2) * sizeof(int16);
+		CompressedBlockSize = BlockSize;
+
+		const uint32 uncompressedBlockSamples = (2 + (BlockSize - PreambleSize) * 2);
+		const uint32 targetBlocks = MONO_PCM_BUFFER_SAMPLES / uncompressedBlockSamples;
+		StreamBufferSize = targetBlocks * UncompressedBlockSize;
+		// Ensure TotalDecodedSize is a even multiple of the compressed block size so that the buffer is not over read on the last block
+		TotalDecodedSize = ((WaveInfo.SampleDataSize + CompressedBlockSize - 1) / CompressedBlockSize) * UncompressedBlockSize;
+
+		UncompressedBlockData = (uint8*)FMemory::Realloc(UncompressedBlockData, NumChannels * UncompressedBlockSize);
+		check(UncompressedBlockData != nullptr);
+		TotalCompressedBlocksPerChannel = (WaveInfo.SampleDataSize + CompressedBlockSize - 1) / CompressedBlockSize / NumChannels;
+	}
+	else if(Format == WAVE_FORMAT_LPCM)
+	{
+		// There are no "blocks" in this case
+		BlockSize = 0;
+		UncompressedBlockSize = 0;
+		CompressedBlockSize = 0;
+		StreamBufferSize = 0;
+		UncompressedBlockData = nullptr;
+		TotalCompressedBlocksPerChannel = 0;
+
+		TotalDecodedSize = WaveInfo.SampleDataSize;
+
+		TotalSamplesPerChannel = TotalDecodedSize / sizeof(uint16) / NumChannels;
+	}
+	else
+	{
+		return false;
+	}
+
+	if (QualityInfo)
+	{
+		QualityInfo->SampleRate = *WaveInfo.pSamplesPerSec;
+		QualityInfo->NumChannels = *WaveInfo.pChannels;
+		QualityInfo->SampleDataSize = TotalDecodedSize;
+		QualityInfo->Duration = (float)TotalSamplesPerChannel / QualityInfo->SampleRate;
+	}
+
+	CurrentCompressedBlockIndex = 0;
+	TotalSamplesStreamed = 0;
+	// This is set to the max value to trigger the decompression of the first audio block
+	CurrentUncompressedBlockSampleIndex = UncompressedBlockSize / sizeof(uint16);
+
+	return true;
+}
+
+bool FADPCMAudioInfo::ReadCompressedData(uint8* Destination, bool bLooping, uint32 BufferSize)
+{
+
+	const uint32 ChannelSampleSize = sizeof(uint16) * NumChannels;
+
+	// This correctly handles any BufferSize as long as its a multiple of sample size * number of channels
+	check(Destination);
+	check(BufferSize % ChannelSampleSize == 0);
+
+	ProcessSeekRequest();
+
+	// If we've already read through this asset and we are not looping, memzero and early out.
+	if (TotalSamplesStreamed >= TotalSamplesPerChannel && !bLooping)
+	{
+		FMemory::Memzero(Destination, BufferSize);
+		return true;
+	}
+
+	int16* OutData = (int16*)Destination;
+	bool ReachedEndOfSamples = false;
+	if(Format == WAVE_FORMAT_ADPCM)
+	{
+		// We need to loop over the number of samples requested since an uncompressed block will not match the number of frames requested
+		while(BufferSize > 0)
+		{
+			if(CurrentUncompressedBlockSampleIndex >= UncompressedBlockSize / sizeof(uint16))
+			{
+				// we need to decompress another block of compressed data from the current chunk
+
+				// Decompress one block for each channel and store it in UncompressedBlockData
+				for(int32 ChannelItr = 0; ChannelItr < NumChannels; ++ChannelItr)
+				{
+					ADPCM::DecodeBlock(
+						WaveInfo.SampleDataStart + (ChannelItr * TotalCompressedBlocksPerChannel + CurrentCompressedBlockIndex) * CompressedBlockSize,
+						CompressedBlockSize,
+						(int16*)(UncompressedBlockData + ChannelItr * UncompressedBlockSize));
+				}
+
+				// Update some bookkeeping
+				CurrentUncompressedBlockSampleIndex = 0;
+				++CurrentCompressedBlockIndex;
+			}
+
+			// Only copy over the number of samples we currently have available, we will loop around if needed
+			uint32 DecompressedSamplesToCopy = FMath::Min<uint32>(UncompressedBlockSize / sizeof(uint16) - CurrentUncompressedBlockSampleIndex, BufferSize / (sizeof(uint16) * NumChannels));
+			check(DecompressedSamplesToCopy > 0);
+
+			// Ensure we don't go over the number of samples left in the audio data
+			if(DecompressedSamplesToCopy > TotalSamplesPerChannel - TotalSamplesStreamed)
+			{
+				DecompressedSamplesToCopy = TotalSamplesPerChannel - TotalSamplesStreamed;
+			}
+
+			// Copy over the actual sample data
+			for(uint32 SampleItr = 0; SampleItr < DecompressedSamplesToCopy; ++SampleItr)
+			{
+				for(int32 ChannelItr = 0; ChannelItr < NumChannels; ++ChannelItr)
+				{
+					uint16 Value = *(int16*)(UncompressedBlockData + ChannelItr * UncompressedBlockSize + (CurrentUncompressedBlockSampleIndex + SampleItr) * sizeof(int16));
+					*OutData++ = Value;
+				}
+			}
+
+			// Update bookkeeping
+			CurrentUncompressedBlockSampleIndex += DecompressedSamplesToCopy;
+			BufferSize -= DecompressedSamplesToCopy * sizeof(uint16) * NumChannels;
+			TotalSamplesStreamed += DecompressedSamplesToCopy;
+
+			// Check for the end of the audio samples and loop if needed
+			if(TotalSamplesStreamed >= TotalSamplesPerChannel)
+			{
+				ReachedEndOfSamples = true;
+				if(!bLooping)
+				{
+					// Zero remaining buffer
+					FMemory::Memzero(OutData, BufferSize);
+					
+					ResetSeekState();
+					
+					return true;
+				}
+				else
+				{
+					// This is set to the max value to trigger the decompression of the first audio block
+					CurrentUncompressedBlockSampleIndex = UncompressedBlockSize / sizeof(uint16);
+					CurrentCompressedBlockIndex = 0;
+					TotalSamplesStreamed = 0;
+				}
+			}
+		}
+	}
+	else
+	{
+		uint32 OutDataOffset = 0;
+		while (BufferSize > 0)
+		{
+			uint32 DecompressedSamplesToCopy = BufferSize / ChannelSampleSize;
+
+			// Ensure we don't go over the number of samples left in the audio data
+			if (DecompressedSamplesToCopy > TotalSamplesPerChannel - TotalSamplesStreamed)
+			{
+				DecompressedSamplesToCopy = TotalSamplesPerChannel - TotalSamplesStreamed;
+			}
+
+			FMemory::Memcpy(OutData + OutDataOffset, WaveInfo.SampleDataStart + (TotalSamplesStreamed * ChannelSampleSize), DecompressedSamplesToCopy * ChannelSampleSize);
+			TotalSamplesStreamed += DecompressedSamplesToCopy;
+			BufferSize -= DecompressedSamplesToCopy * ChannelSampleSize;
+			OutDataOffset += DecompressedSamplesToCopy * NumChannels;
+
+			// Check for the end of the audio samples and loop if needed
+			if (TotalSamplesStreamed >= TotalSamplesPerChannel)
+			{
+				ReachedEndOfSamples = true;
+				TotalSamplesStreamed = 0;
+				if (!bLooping)
+				{
+					// Zero remaining buffer
+					FMemory::Memzero(OutData + OutDataOffset, BufferSize);
+
+					ResetSeekState();
+
+					return true;
+				}
+			}
+		}
+	}
+
+	ResetSeekState();
+
+	return ReachedEndOfSamples;
+}
+
+void FADPCMAudioInfo::ExpandFile(uint8* DstBuffer, struct FSoundQualityInfo* QualityInfo)
+{
+	check(DstBuffer);
+
+	ReadCompressedData(DstBuffer, false, TotalDecodedSize);
+}
+
+int FADPCMAudioInfo::GetStreamBufferSize() const
+{
+	return StreamBufferSize;
+}
+
+void FADPCMAudioInfo::ProcessSeekRequest()
+{
+	if (bNewSeekRequest)
+	{
+		SeekToFrameInternal(TargetSeekFrame);
+	}
+}
+
+void FADPCMAudioInfo::ResetSeekState()
+{
+	bSeekPendingRead = false;
+	bNewSeekRequest = false;
+}
+
+bool FADPCMAudioInfo::StreamCompressedInfoInternal(const FSoundWaveProxyPtr& InWaveProxy, struct FSoundQualityInfo* QualityInfo)
+{
+	FScopeLock ScopeLock(&CurCompressedChunkHandleCriticalSection);
+
+	check(QualityInfo);
+	check(StreamingSoundWave == InWaveProxy);
+	if (!ensure(InWaveProxy.IsValid()))
+	{
+		return false;
+	}
+
+	CurrentChunkIndex = 0;
+
+	// Get the first chunk of audio data (should already be loaded)
+	uint8 const* ChunkData = GetLoadedChunk(InWaveProxy, CurrentChunkIndex, CurrentChunkDataSize);
+
+	if (ChunkData == nullptr)
+	{
+		UE_LOG(LogAudio, Warning, TEXT("FADPCMAudioInfo::StreamCompressedInfoInternal: Failed to get loaded chunk at index '%u'"), CurrentChunkIndex);
+		return false;
+	}
+
+	SrcBufferData = nullptr;
+	SrcBufferDataSize = 0;
+
+	void* FormatHeader = nullptr;
+	FString ErrorMsg;
+	constexpr bool bHeaderDataOnly = true;
+	if (!WaveInfo.ReadWaveInfo(static_cast<const uint8*>(ChunkData), CurrentChunkDataSize, &ErrorMsg, bHeaderDataOnly, &FormatHeader))
+	{
+		UE_LOG(LogAudio, Warning, TEXT("FADPCMAudioInfo::StreamCompressedInfoInternal: Failed to read WaveInfo: %s"), *ErrorMsg);
+		return false;
+	}
+	
+	// if we only included the header in the zeroth chunk, skip to the next chunk.
+	const int32 SampleDataOffset = UE_PTRDIFF_TO_INT32(WaveInfo.SampleDataStart - ChunkData);
+	check(SampleDataOffset > 0);
+	if (((uint32)SampleDataOffset) >= CurrentChunkDataSize)
+	{
+		++CurrentChunkIndex;
+		ChunkData = GetLoadedChunk(InWaveProxy, CurrentChunkIndex, CurrentChunkDataSize);
+		FirstChunkSampleDataIndex = CurrentChunkIndex;
+		FirstChunkSampleDataOffset = 0;
+	}
+	else
+	{
+		FirstChunkSampleDataOffset = UE_PTRDIFF_TO_INT32(WaveInfo.SampleDataStart - ChunkData);
+		FirstChunkSampleDataIndex = 0;
+	}
+
+	PreviouslyRequestedChunkIndex = CurrentChunkIndex;
+
+	SrcBufferData = ChunkData;
+	CurrentChunkBufferOffset = 0;
+	CurCompressedChunkData = nullptr;
+	bDecompressorReleased = false;
+	CurrentUncompressedBlockSampleIndex = 0;
+	
+	TotalSamplesStreamed = 0;
+	Format = *WaveInfo.pFormatTag;
+	NumChannels = *WaveInfo.pChannels;
+
+	if (Format == WAVE_FORMAT_ADPCM)
+	{
+		check(FormatHeader);
+		ADPCM::ADPCMFormatHeader* ADPCMHeader = (ADPCM::ADPCMFormatHeader*)FormatHeader;
+		TotalSamplesPerChannel = ADPCMHeader->SamplesPerChannel;
+		SamplesPerBlock = ADPCMHeader->wSamplesPerBlock;
+
+		const uint32 PreambleSize = 7;
+
+		BlockSize = *WaveInfo.pBlockAlign;
+		UncompressedBlockSize = (2 + (BlockSize - PreambleSize) * 2) * sizeof(int16);
+		CompressedBlockSize = BlockSize;
+
+		// Calculate buffer sizes and total number of samples
+		const uint32 uncompressedBlockSamples = (2 + (BlockSize - PreambleSize) * 2);
+		const uint32 targetBlocks = MONO_PCM_BUFFER_SAMPLES / uncompressedBlockSamples;
+		StreamBufferSize = targetBlocks * UncompressedBlockSize;
+		TotalDecodedSize = ((WaveInfo.SampleDataSize + CompressedBlockSize - 1) / CompressedBlockSize) * UncompressedBlockSize;
+
+		UncompressedBlockData = (uint8*)FMemory::Realloc(UncompressedBlockData, NumChannels * UncompressedBlockSize);
+		check(UncompressedBlockData != nullptr);
+	}
+	else if (Format == WAVE_FORMAT_LPCM)
+	{
+		BlockSize = 0;
+		UncompressedBlockSize = 0;
+		CompressedBlockSize = 0;
+
+		// This is uncompressed, so decoded size and buffer size are the same
+		TotalDecodedSize = WaveInfo.SampleDataSize;
+		StreamBufferSize = WaveInfo.SampleDataSize;
+		TotalSamplesPerChannel = StreamBufferSize / sizeof(uint16) / NumChannels;
+	}
+	else
+	{
+		UE_LOG(LogAudio, Error, TEXT("FADPCMAudioInfo::StreamCompressedInfoInternal: Unsupported wave format (Neither ADPCM nor LPCM)"));
+		return false;
+	}
+
+	if (QualityInfo)
+	{
+		QualityInfo->SampleRate = *WaveInfo.pSamplesPerSec;
+		QualityInfo->NumChannels = *WaveInfo.pChannels;
+		QualityInfo->SampleDataSize = TotalDecodedSize;
+		QualityInfo->Duration = (float)TotalSamplesPerChannel / QualityInfo->SampleRate;
+	}
+
+	return true;
+}
+
+bool FADPCMAudioInfo::StreamCompressedData(uint8* Destination, bool bLooping, uint32 BufferSize, int32& OutNumBytesStreamed)
+{
+	FScopeLock ScopeLock(&CurCompressedChunkHandleCriticalSection);
+	int32 OriginalBufferSize = (int32)BufferSize;
+	OutNumBytesStreamed = 0;
+
+	// Initial sanity checks:
+	if (bDecompressorReleased)
+	{
+		UE_LOG(LogAudio, Error, TEXT("Stream Compressed Data called after release!"));
+		if (Destination != nullptr)
+		{
+			FMemory::Memzero(Destination, BufferSize);
+		}
+		return true;
+	}
+
+	if (Destination == nullptr || BufferSize == 0)
+	{
+		UE_LOG(LogAudio, Error, TEXT("Stream Compressed Info not called!"));
+		return false;
+	}
+
+	if (NumChannels == 0)
+	{
+		UE_LOG(LogAudio, Error, TEXT("Stream Compressed Info not called!"));
+		FMemory::Memzero(Destination, BufferSize);
+		return true;
+	}
+
+	// Destination samples are interlaced by channel, BufferSize is in bytes
+	const int32 ChannelSampleSize = sizeof(uint16) * NumChannels;
+
+	// Ensure that BuffserSize is a multiple of the sample size times the number of channels
+	if (BufferSize % ChannelSampleSize != 0)
+	{
+		UE_LOG(LogAudio, Error, TEXT("Invalid buffer size %d requested for %d channels."), BufferSize, NumChannels);
+		FMemory::Memzero(Destination, BufferSize);
+		return true;
+	}
+
+	ProcessSeekRequest();
+
+	int16* OutData = (int16*)Destination;
+	bool ReachedEndOfSamples = false;
+	if(Format == WAVE_FORMAT_ADPCM)
+	{
+		// We need to loop over the number of samples requested since an uncompressed block will not match the number of frames requested
+		while(BufferSize > 0)
+		{
+			if(CurCompressedChunkData == nullptr || CurrentUncompressedBlockSampleIndex >= UncompressedBlockSize / sizeof(uint16))
+			{
+				// we need to decompress another block of compressed data from the current chunk
+
+				if(CurCompressedChunkData == nullptr || CurrentChunkBufferOffset >= CurrentChunkDataSize)
+				{
+					// Get another chunk of compressed data from the streaming engine
+
+					// CurrentChunkIndex is used to keep track of the current chunk for loading/unloading by the streaming engine, but chunk 0 is
+					// preloaded so don't increment this when getting chunk 0, only later chunks. If previous chunk retrieval failed because it was
+					// not ready, don't re-increment the CurrentChunkIndex.  Only increment if seek not pending as seek determines the current chunk index
+					// and invalidates CurCompressedChunkData.
+					if(CurCompressedChunkData != nullptr)
+					{
+						// If we have already seeked forward to the next chunk because the chunk didn't load in time for the render callback, we don't increment here.
+						// Otherwise, we increment to the next chunk index and reset our chunk offset to the beginning of the new chunk.
+						if (!bSeekedFowardToNextChunk)
+						{
+							// Ensure that we're either seeking or moving sequentially through the file.
+							check(bSeekPendingRead || CurrentChunkIndex == PreviouslyRequestedChunkIndex);
+							++CurrentChunkIndex;
+
+							// Set the current buffer offset accounting for the header in the first chunk
+							if (!bSeekPendingRead)
+							{
+								// Ensure that, if we hit this code, we are not at the beginning of the file.
+								check(CurrentChunkIndex != FirstChunkSampleDataIndex);
+								CurrentChunkBufferOffset = 0;
+							}
+						}
+
+						// If this flag was raised, clear it for the next time we move on to another chunk.
+						bSeekedFowardToNextChunk = false;
+					}
+
+					if (CurrentChunkIndex >= StreamingSoundWave->GetNumChunks())
+					{
+						// If this is the case, we've seeked to the end of the file.
+						ReachedEndOfSamples = true;
+						CurrentUncompressedBlockSampleIndex = 0;
+						CurrentChunkIndex = FirstChunkSampleDataIndex;
+						CurrentChunkBufferOffset = FirstChunkSampleDataOffset;
+						TotalSamplesStreamed = 0;
+						CurCompressedChunkData = nullptr;
+						if (!bLooping)
+						{
+							// Set the remaining buffer to 0
+							FMemory::Memset(OutData, 0, BufferSize);
+							return true;
+						}
+					}
+
+					// Request the next chunk of data from the streaming engine
+					CurCompressedChunkData = GetLoadedChunk(StreamingSoundWave, CurrentChunkIndex, CurrentChunkDataSize);
+
+					if(CurCompressedChunkData == nullptr)
+					{
+						// We only need to worry about missing the stream chunk if we were seeking. Seeking might cause a bit of latency with chunk loading. That is expected.
+						if (!bSeekPendingRead)
+						{
+							// If we did not load chunk then bail. CurrentChunkIndex will not get incremented on the next callback so in effect another attempt will be made to fetch the chunk.
+							// Since audio streaming depends on the general data streaming mechanism used by other parts of the engine and new data is pre-fetched on the game thread, its
+							// possible a game thread stall can cause this.
+							UE_LOG(LogAudio, Verbose, TEXT("Missed Deadline chunk %d"), CurrentChunkIndex);
+						}
+
+						// If we have a seek pending, or we're in the middle of playing back the file, seek forward in the chunk offset to where we would have been.
+						// If we've already seeked forward one chunk and haven't loaded the next chunk in time and are skipping to the next one, we stop seeking to prevent snowballing chunk load requests.
+						const bool bAlreadySeekedFowardOneChunk = ADPCMOnlySeekForwardOneChunkCVar && bSeekedFowardToNextChunk;
+						const bool bShouldSeekForwardOnMissedChunk = !ADPCMDisableSeekForwardOnChunkMissesCVar && !bAlreadySeekedFowardOneChunk && (bSeekPendingRead || CurrentChunkIndex > FirstChunkSampleDataIndex);
+
+						if (bShouldSeekForwardOnMissedChunk)
+						{
+							const uint32 NumBlocksForCallback = BufferSize / (UncompressedBlockSize * NumChannels);
+							const uint32 NumCompressedBytesForCallback = NumBlocksForCallback * CompressedBlockSize * NumChannels;
+							
+							uint32 NumDecompressedSamplesForCallback = BufferSize / (ChannelSampleSize);
+
+							if (NumDecompressedSamplesForCallback > TotalSamplesPerChannel - TotalSamplesStreamed)
+							{
+								// If this is the case, we've seeked to the end of the file.
+								ReachedEndOfSamples = true;
+								CurrentUncompressedBlockSampleIndex = 0;
+								CurrentChunkIndex = FirstChunkSampleDataIndex;
+								CurrentChunkBufferOffset = FirstChunkSampleDataOffset;
+								TotalSamplesStreamed = 0;
+								CurCompressedChunkData = nullptr;
+								if (!bLooping)
+								{
+									// Set the remaining buffer to 0
+									FMemory::Memset(OutData, 0, BufferSize);
+									return true;
+								}
+							}
+
+							// Seek forward in the file the amount we would be rendering otherwise.
+							CurrentChunkBufferOffset += NumCompressedBytesForCallback;
+							TotalSamplesStreamed += NumDecompressedSamplesForCallback;
+
+							const uint32 SizeOfCurrentChunk = StreamingSoundWave->GetSizeOfChunk(CurrentChunkIndex);
+							if (CurrentChunkBufferOffset >= SizeOfCurrentChunk)
+							{
+								// If this is the case, we've seeked forward past the end of the chunk.
+
+								// Ensure we never hit this section of code twice, resulting in a double increment of the chunk index.
+								ensureAlways(!bSeekedFowardToNextChunk);
+
+								// Ensure that we are either seeking of moving sequentially through the file.
+								ensureAlwaysMsgf(bSeekPendingRead || CurrentChunkIndex == PreviouslyRequestedChunkIndex, TEXT("Failed to load a chunk of ADPCM audio for the entire duration that chunk of audio was supposed to be played."));
+								++CurrentChunkIndex;
+								CurrentChunkBufferOffset -= SizeOfCurrentChunk;
+
+								// Ensure that we have not moved past the size of our new chunk with our current callback.
+								ensureAlways(CurrentChunkBufferOffset < StreamingSoundWave->GetSizeOfChunk(CurrentChunkIndex));
+
+								bSeekedFowardToNextChunk = true;
+
+								if (!ensureMsgf(CurrentChunkBufferOffset % CompressedBlockSize == 0, TEXT("Error: seeked partway into an ADPCM block. Please check the above code, as well as FAudioFormatADPCM::SplitDataForStreaming for accuracy.")))
+								{
+									CurrentChunkBufferOffset = AlignDown(CurrentChunkBufferOffset, CompressedBlockSize);
+								}
+
+								if (CurrentChunkIndex >= StreamingSoundWave->GetNumChunks())
+								{
+									// If this is the case, we've seeked to the end of the file.
+									ReachedEndOfSamples = true;
+									CurrentUncompressedBlockSampleIndex = 0;
+									CurrentChunkIndex = FirstChunkSampleDataIndex;
+									CurrentChunkBufferOffset = FirstChunkSampleDataOffset;
+									TotalSamplesStreamed = 0;
+									CurCompressedChunkData = nullptr;
+									if (!bLooping)
+									{
+										// Set the remaining buffer to 0
+										FMemory::Memset(OutData, 0, BufferSize);
+										return true;
+									}
+								}
+							}
+						}
+
+						// zero out remaining data and bail
+						FMemory::Memset(OutData, 0, BufferSize);
+						return false;
+					}
+
+					
+					if (CurrentChunkIndex == FirstChunkSampleDataIndex && !bSeekPendingRead)
+					{
+						// If we're in the first chunk, set the current buffer offset accounting for the header.
+						CurrentChunkBufferOffset = FirstChunkSampleDataOffset;
+					}
+
+					ResetSeekState();
+				}
+
+				// Decompress one block for each channel and store it in UncompressedBlockData
+				for(int32 ChannelItr = 0; ChannelItr < NumChannels; ++ChannelItr)
+				{
+					ADPCM::DecodeBlock(
+						CurCompressedChunkData + CurrentChunkBufferOffset + ChannelItr * CompressedBlockSize,
+						CompressedBlockSize,
+						(int16*)(UncompressedBlockData + ChannelItr * UncompressedBlockSize));
+				}
+
+				// Update some bookkeeping
+				CurrentUncompressedBlockSampleIndex = 0;
+				CurrentChunkBufferOffset += NumChannels * CompressedBlockSize;
+			}
+
+			// Only copy over the number of samples we currently have available, we will loop around if needed
+			uint32 DecompressedSamplesToCopy = FMath::Min<uint32>(
+				(UncompressedBlockSize / sizeof(uint16)) - CurrentUncompressedBlockSampleIndex,
+				BufferSize / (ChannelSampleSize));
+			check(DecompressedSamplesToCopy > 0);
+
+			// Ensure we don't go over the number of samples left in the audio data
+			if(DecompressedSamplesToCopy > TotalSamplesPerChannel - TotalSamplesStreamed)
+			{
+				DecompressedSamplesToCopy = TotalSamplesPerChannel - TotalSamplesStreamed;
+			}
+
+			// Copy over the actual sample data
+			for(uint32 SampleItr = 0; SampleItr < DecompressedSamplesToCopy; ++SampleItr)
+			{
+				for(int32 ChannelItr = 0; ChannelItr < NumChannels; ++ChannelItr)
+				{
+					uint16 Value = *(int16*)(UncompressedBlockData + ChannelItr * UncompressedBlockSize + (CurrentUncompressedBlockSampleIndex + SampleItr) * sizeof(int16));
+					*OutData++ = Value;
+					OutNumBytesStreamed += sizeof(int16);
+				}
+			}
+
+			// Update bookkeeping
+			CurrentUncompressedBlockSampleIndex += DecompressedSamplesToCopy;
+			BufferSize -= DecompressedSamplesToCopy * ChannelSampleSize;
+			TotalSamplesStreamed += DecompressedSamplesToCopy;
+
+			// Check for the end of the audio samples and loop if needed
+			if(TotalSamplesStreamed >= TotalSamplesPerChannel)
+			{
+				ReachedEndOfSamples = true;
+				CurrentUncompressedBlockSampleIndex = 0;
+				CurrentChunkIndex = FirstChunkSampleDataIndex;
+				CurrentChunkBufferOffset = 0;
+				TotalSamplesStreamed = 0;
+				CurCompressedChunkData = nullptr;
+				if(!bLooping)
+				{
+					// Set the remaining buffer to 0
+					FMemory::Memset(OutData, 0, BufferSize);
+					return true;
+				}
+			}
+		}
+	}
+	else
+	{
+		while(BufferSize > 0)
+		{
+			if(CurCompressedChunkData == nullptr || CurrentChunkBufferOffset >= CurrentChunkDataSize)
+			{
+				// Get another chunk of compressed data from the streaming engine
+
+				// CurrentChunkIndex is used to keep track of the current chunk for loading/unloading by the streaming engine but chunk 0
+				// is preloaded so we don't want to increment this when getting chunk 0, only later chunks
+				// Also, if we failed to get a previous chunk because it was not ready we don't want to re-increment the CurrentChunkIndex
+				if(CurCompressedChunkData != nullptr)
+				{
+					++CurrentChunkIndex;
+				}
+
+				// Check for the end of the audio samples and loop if needed
+				if (CurrentChunkIndex >= StreamingSoundWave->GetNumChunks())
+				{
+					ReachedEndOfSamples = true;
+					CurrentChunkIndex = FirstChunkSampleDataIndex;
+					CurrentChunkBufferOffset = FirstChunkSampleDataOffset;
+					TotalSamplesStreamed = 0;
+					CurCompressedChunkData = nullptr;
+					if (!bLooping)
+					{
+						// Set the remaining buffer to 0
+						FMemory::Memset(OutData, 0, BufferSize);
+						return true;
+					}
+				}
+
+				// Request the next chunk of data from the streaming engine
+				CurCompressedChunkData = GetLoadedChunk(StreamingSoundWave, CurrentChunkIndex, CurrentChunkDataSize);
+
+				if(CurCompressedChunkData == nullptr)
+				{
+					// Only report missing the stream chunk if we were seeking. Seeking
+					// may cause a bit of latency with chunk loading, which is expected.
+					if (!bSeekPendingRead)
+					{
+						// CurrentChunkIndex will not get incremented on the next callback, effectively causing another attempt to fetch the chunk.
+						// Since audio streaming depends on the general data streaming mechanism used by other parts of the engine and new data is
+						// prefetched on the game tick thread, this may be caused by a game hitch.
+						UE_LOG(LogAudio, Verbose, TEXT("Missed streaming ADPCM deadline chunk %d"), CurrentChunkIndex);
+					}
+
+					FMemory::Memset(OutData, 0, BufferSize);
+					NumConsecutiveReadFailiures++;
+					const bool bReadAttemptTimedOut = NumConsecutiveReadFailiures > ADPCMReadFailiureTimeoutCVar;
+					UE_CLOG(bReadAttemptTimedOut, LogAudio, Warning, TEXT("ADPCM Audio Decode timed out."), bReadAttemptTimedOut);
+					return NumConsecutiveReadFailiures > ADPCMReadFailiureTimeoutCVar;
+				}
+
+				NumConsecutiveReadFailiures = 0;
+
+				// Set the current buffer offset accounting for the header in the first chunk
+				if (!bSeekPendingRead)
+				{
+					CurrentChunkBufferOffset = CurrentChunkIndex == FirstChunkSampleDataIndex ? FirstChunkSampleDataOffset : 0;
+				}
+
+				ResetSeekState();
+			}
+			
+			uint32 DecompressedSamplesToCopy = FMath::Min<uint32>(
+				(CurrentChunkDataSize - CurrentChunkBufferOffset) / ChannelSampleSize,
+				BufferSize / ChannelSampleSize);
+
+			if (DecompressedSamplesToCopy == 0)
+			{
+				CurrentChunkBufferOffset = CurrentChunkDataSize;
+				continue;
+			}
+
+			// Ensure we don't go over the number of samples left in the audio data
+			if(DecompressedSamplesToCopy > TotalSamplesPerChannel - TotalSamplesStreamed)
+			{
+				DecompressedSamplesToCopy = TotalSamplesPerChannel - TotalSamplesStreamed;
+			}
+
+			const uint32 SampleSize = DecompressedSamplesToCopy * ChannelSampleSize;
+			FMemory::Memcpy(OutData, CurCompressedChunkData + CurrentChunkBufferOffset, SampleSize);
+			OutData += DecompressedSamplesToCopy * NumChannels;
+			OutNumBytesStreamed += DecompressedSamplesToCopy * NumChannels * sizeof(int16);
+
+			CurrentChunkBufferOffset += SampleSize;
+			BufferSize -= SampleSize;
+			TotalSamplesStreamed += DecompressedSamplesToCopy;
+
+			// Check for the end of the audio samples and loop if needed
+			if(TotalSamplesStreamed >= TotalSamplesPerChannel)
+			{
+				ReachedEndOfSamples = true;
+				CurrentChunkIndex = FirstChunkSampleDataIndex;
+				CurrentChunkBufferOffset = FirstChunkSampleDataOffset;
+				TotalSamplesStreamed = 0;
+				CurCompressedChunkData = nullptr;
+				if(!bLooping)
+				{
+					// Set the remaining buffer to 0
+					FMemory::Memset(OutData, 0, BufferSize);
+					return true;
+				}
+			}
+		}
+	}
+
+	return ReachedEndOfSamples;
+}
+
+bool FADPCMAudioInfo::ReleaseStreamChunk(bool bBlockUntilReleased)
+{
+	ResetSeekState();
+
+	if (bBlockUntilReleased)
+	{
+		// Wait for any pending decode tasks to finish.
+		FScopeLock ScopeLock(&CurCompressedChunkHandleCriticalSection);
+		CurCompressedChunkHandle = FAudioChunkHandle();
+		CurrentChunkBufferOffset = 0;
+		CurCompressedChunkData = nullptr;
+		CurrentChunkDataSize = 0;
+		bDecompressorReleased = true;
+		return true;
+	}
+	else
+	{
+		// If the chunk isn't currently in use, reset the chunk handle.
+		if (CurCompressedChunkHandleCriticalSection.TryLock())
+		{
+			CurCompressedChunkHandle = FAudioChunkHandle();
+			CurrentChunkBufferOffset = 0;
+			CurCompressedChunkData = nullptr;
+			CurrentChunkDataSize = 0;
+			bDecompressorReleased = true;
+			CurCompressedChunkHandleCriticalSection.Unlock();
+			return true;
+		}
+		else
+		{
+			// Otherwise, do nothing and exit.
+			return false;
+		}
+	}
+}
+
+const uint8* FADPCMAudioInfo::GetLoadedChunk(const FSoundWaveProxyPtr& InSoundWave, uint32 ChunkIndex, uint32& OutChunkSize)
+{
+	if (ChunkIndex != 0 && FMath::RandRange(0.0f, 1.0f) < ChanceForIntentionalChunkMissCVar)
+	{
+		UE_LOG(LogAudio, Display, TEXT("Intentionally dropping a chunk here."));
+		OutChunkSize = 0;
+		return nullptr;
+	}
+
+	if (!ensure(InSoundWave.IsValid()) || ChunkIndex >= InSoundWave->GetNumChunks())
+	{
+		OutChunkSize = 0;
+		return nullptr;
+	}
+	else if (ChunkIndex == 0)
+	{
+		TArrayView<const uint8> ZerothChunk = FSoundWaveProxy::GetZerothChunk(InSoundWave, true);
+		OutChunkSize = ZerothChunk.Num();
+		PreviouslyRequestedChunkIndex = ChunkIndex;
+		return ZerothChunk.GetData();
+	}
+	else
+	{
+		const bool bIsSeekingOrLooping = bSeekPendingRead || ChunkIndex == FirstChunkSampleDataIndex;
+		ensureAlwaysMsgf(bIsSeekingOrLooping || ChunkIndex == PreviouslyRequestedChunkIndex || ChunkIndex == PreviouslyRequestedChunkIndex + 1, TEXT("ADPCM playback error! We skipped from the end of chunk %d to the beginning of chunk %d."), PreviouslyRequestedChunkIndex, ChunkIndex);
+
+		CurCompressedChunkHandle = IStreamingManager::Get().GetAudioStreamingManager().GetLoadedChunk(InSoundWave, ChunkIndex, false, true);
+		OutChunkSize = CurCompressedChunkHandle.Num();
+
+		PreviouslyRequestedChunkIndex = ChunkIndex;
+		
+		return CurCompressedChunkHandle.GetData();
+	}
+}
+
+//////////////////////////////////////////////////////////////////////////
+// Copied from IOS - probably want to split and share
+//////////////////////////////////////////////////////////////////////////
+#define NUM_ADAPTATION_TABLE 16
+#define NUM_ADAPTATION_COEFF 7
+#define SOUND_SOURCE_FREE 0
+#define SOUND_SOURCE_LOCKED 1
+
+namespace
+{
+	template <typename T, uint32 B>
+	inline T SignExtend(const T ValueToExtend)
+	{
+		struct { T ExtendedValue : B; } SignExtender;
+		return SignExtender.ExtendedValue = ValueToExtend;
+	}
+
+	template <typename T>
+	inline T ReadFromByteStream(const uint8* ByteStream, int32& ReadIndex, bool bLittleEndian = true)
+	{
+		T ValueRaw = 0;
+
+		if (bLittleEndian)
+		{
+#if PLATFORM_LITTLE_ENDIAN
+			for (int32 ByteIndex = 0; ByteIndex < sizeof(T); ++ByteIndex)
+#else
+			for (int32 ByteIndex = sizeof(T) - 1; ByteIndex >= 0; --ByteIndex)
+#endif // PLATFORM_LITTLE_ENDIAN
+			{
+				ValueRaw |= ByteStream[ReadIndex++] << 8 * ByteIndex;
+			}
+		}
+		else
+		{
+#if PLATFORM_LITTLE_ENDIAN
+			for (int32 ByteIndex = sizeof(T) - 1; ByteIndex >= 0; --ByteIndex)
+#else
+			for (int32 ByteIndex = 0; ByteIndex < sizeof(T); ++ByteIndex)
+#endif // PLATFORM_LITTLE_ENDIAN
+			{
+				ValueRaw |= ByteStream[ReadIndex++] << 8 * ByteIndex;
+			}
+		}
+
+		return ValueRaw;
+	}
+
+	template <typename T>
+	inline void WriteToByteStream(T Value, uint8* ByteStream, int32& WriteIndex, bool bLittleEndian = true)
+	{
+		if (bLittleEndian)
+		{
+#if PLATFORM_LITTLE_ENDIAN
+			for (int32 ByteIndex = 0; ByteIndex < sizeof(T); ++ByteIndex)
+#else
+			for (int32 ByteIndex = sizeof(T) - 1; ByteIndex >= 0; --ByteIndex)
+#endif // PLATFORM_LITTLE_ENDIAN
+			{
+				ByteStream[WriteIndex++] = (Value >> (8 * ByteIndex)) & 0xFF;
+			}
+		}
+		else
+		{
+#if PLATFORM_LITTLE_ENDIAN
+			for (int32 ByteIndex = sizeof(T) - 1; ByteIndex >= 0; --ByteIndex)
+#else
+			for (int32 ByteIndex = 0; ByteIndex < sizeof(T); ++ByteIndex)
+#endif // PLATFORM_LITTLE_ENDIAN
+			{
+				ByteStream[WriteIndex++] = (Value >> (8 * ByteIndex)) & 0xFF;
+			}
+		}
+	}
+
+	template <typename T>
+	inline T ReadFromArray(const T* ElementArray, int32& ReadIndex, int32 NumElements, int32 IndexStride = 1)
+	{
+		T OutputValue = 0;
+
+		if (ReadIndex >= 0 && ReadIndex < NumElements)
+		{
+			OutputValue = ElementArray[ReadIndex];
+			ReadIndex += IndexStride;
+		}
+
+		return OutputValue;
+	}
+
+	inline bool LockSourceChannel(int32* ChannelLock)
+	{
+		check(ChannelLock != NULL);
+		return FPlatformAtomics::InterlockedCompareExchange(ChannelLock, SOUND_SOURCE_LOCKED, SOUND_SOURCE_FREE) == SOUND_SOURCE_FREE;
+	}
+
+	inline void UnlockSourceChannel(int32* ChannelLock)
+	{
+		check(ChannelLock != NULL);
+		int32 Result = FPlatformAtomics::InterlockedCompareExchange(ChannelLock, SOUND_SOURCE_FREE, SOUND_SOURCE_LOCKED);
+
+		check(Result == SOUND_SOURCE_LOCKED);
+	}
+
+} // end namespace
+
+
+//////////////////////////////////////////////////////////////////////////
+// End of copied section
+//////////////////////////////////////////////////////////////////////////
+
+
+namespace ADPCM
+{
+	template <typename T>
+	static void GetAdaptationTable(T(&OutAdaptationTable)[NUM_ADAPTATION_TABLE])
+	{
+		// Magic values as specified by standard
+		static T AdaptationTable[] =
+		{
+			230, 230, 230, 230, 307, 409, 512, 614,
+			768, 614, 512, 409, 307, 230, 230, 230
+		};
+
+		FMemory::Memcpy(&OutAdaptationTable, AdaptationTable, sizeof(AdaptationTable));
+	}
+
+	struct FAdaptationContext
+	{
+	public:
+		// Adaptation constants
+		int32 AdaptationTable[NUM_ADAPTATION_TABLE];
+		int32 AdaptationCoefficient1[NUM_ADAPTATION_COEFF];
+		int32 AdaptationCoefficient2[NUM_ADAPTATION_COEFF];
+
+		int32 AdaptationDelta;
+		int32 Coefficient1;
+		int32 Coefficient2;
+		int32 Sample1;
+		int32 Sample2;
+
+		FAdaptationContext() :
+			AdaptationDelta(0),
+			Coefficient1(0),
+			Coefficient2(0),
+			Sample1(0),
+			Sample2(0)
+		{
+			GetAdaptationTable(AdaptationTable);
+			GetAdaptationCoefficients(AdaptationCoefficient1, AdaptationCoefficient2);
+		}
+	};
+
+	FORCEINLINE int16 DecodeNibble(FAdaptationContext& Context, uint8 EncodedNibble)
+	{
+		int32 PredictedSample = (Context.Sample1 * Context.Coefficient1 + Context.Sample2 * Context.Coefficient2) / 256;
+		PredictedSample += SignExtend<int8, 4>(EncodedNibble) * Context.AdaptationDelta;
+		PredictedSample = FMath::Clamp(PredictedSample, -32768, 32767);
+
+		// Shuffle samples for the next iteration
+		Context.Sample2 = Context.Sample1;
+		Context.Sample1 = static_cast<int16>(PredictedSample);
+		Context.AdaptationDelta = (Context.AdaptationDelta * Context.AdaptationTable[EncodedNibble]) / 256;
+		Context.AdaptationDelta = FMath::Max(Context.AdaptationDelta, 16);
+
+		return Context.Sample1;
+	}
+
+	void DecodeBlock(const uint8* EncodedADPCMBlock, int32 BlockSize, int16* DecodedPCMData)
+	{
+		FAdaptationContext Context;
+		int32 ReadIndex = 0;
+		int32 WriteIndex = 0;
+
+		uint8 CoefficientIndex = ReadFromByteStream<uint8>(EncodedADPCMBlock, ReadIndex);
+		Context.AdaptationDelta = ReadFromByteStream<int16>(EncodedADPCMBlock, ReadIndex);
+		Context.Sample1 = ReadFromByteStream<int16>(EncodedADPCMBlock, ReadIndex);
+		Context.Sample2 = ReadFromByteStream<int16>(EncodedADPCMBlock, ReadIndex);
+		Context.Coefficient1 = Context.AdaptationCoefficient1[CoefficientIndex];
+		Context.Coefficient2 = Context.AdaptationCoefficient2[CoefficientIndex];
+
+		// The first two samples are sent directly to the output in reverse order, as per the standard
+		DecodedPCMData[WriteIndex++] = Context.Sample2;
+		DecodedPCMData[WriteIndex++] = Context.Sample1;
+
+		uint8 EncodedNibblePair = 0;
+		uint8 EncodedNibble = 0;
+		while (ReadIndex < BlockSize)
+		{
+			// Read from the byte stream and advance the read head.
+			EncodedNibblePair = ReadFromByteStream<uint8>(EncodedADPCMBlock, ReadIndex);
+
+			EncodedNibble = (EncodedNibblePair >> 4) & 0x0F;
+			DecodedPCMData[WriteIndex++] = DecodeNibble(Context, EncodedNibble);
+
+			EncodedNibble = EncodedNibblePair & 0x0F;
+			DecodedPCMData[WriteIndex++] = DecodeNibble(Context, EncodedNibble);
+		}
+	}
+
+	// Decode two PCM streams and interleave as stereo data
+	void DecodeBlockStereo(const uint8* EncodedADPCMBlockLeft, const uint8* EncodedADPCMBlockRight, int32 BlockSize, int16* DecodedPCMData)
+	{
+		FAdaptationContext ContextLeft;
+		FAdaptationContext ContextRight;
+		int32 ReadIndexLeft = 0;
+		int32 ReadIndexRight = 0;
+		int32 WriteIndex = 0;
+
+		uint8 CoefficientIndexLeft = ReadFromByteStream<uint8>(EncodedADPCMBlockLeft, ReadIndexLeft);
+		ContextLeft.AdaptationDelta = ReadFromByteStream<int16>(EncodedADPCMBlockLeft, ReadIndexLeft);
+		ContextLeft.Sample1 = ReadFromByteStream<int16>(EncodedADPCMBlockLeft, ReadIndexLeft);
+		ContextLeft.Sample2 = ReadFromByteStream<int16>(EncodedADPCMBlockLeft, ReadIndexLeft);
+		ContextLeft.Coefficient1 = ContextLeft.AdaptationCoefficient1[CoefficientIndexLeft];
+		ContextLeft.Coefficient2 = ContextLeft.AdaptationCoefficient2[CoefficientIndexLeft];
+
+		uint8 CoefficientIndexRight = ReadFromByteStream<uint8>(EncodedADPCMBlockRight, ReadIndexRight);
+		ContextRight.AdaptationDelta = ReadFromByteStream<int16>(EncodedADPCMBlockRight, ReadIndexRight);
+		ContextRight.Sample1 = ReadFromByteStream<int16>(EncodedADPCMBlockRight, ReadIndexRight);
+		ContextRight.Sample2 = ReadFromByteStream<int16>(EncodedADPCMBlockRight, ReadIndexRight);
+		ContextRight.Coefficient1 = ContextRight.AdaptationCoefficient1[CoefficientIndexRight];
+		ContextRight.Coefficient2 = ContextRight.AdaptationCoefficient2[CoefficientIndexRight];
+
+		// The first two samples from each stream are sent directly to the output in reverse order, as per the standard
+		DecodedPCMData[WriteIndex++] = ContextLeft.Sample2;
+		DecodedPCMData[WriteIndex++] = ContextRight.Sample2;
+		DecodedPCMData[WriteIndex++] = ContextLeft.Sample1;
+		DecodedPCMData[WriteIndex++] = ContextRight.Sample1;
+
+		uint8 EncodedNibblePairLeft = 0;
+		uint8 EncodedNibbleLeft = 0;
+		uint8 EncodedNibblePairRight = 0;
+		uint8 EncodedNibbleRight = 0;
+
+		while (ReadIndexLeft < BlockSize)
+		{
+			// Read from the byte stream and advance the read head.
+			EncodedNibblePairLeft = ReadFromByteStream<uint8>(EncodedADPCMBlockLeft, ReadIndexLeft);
+			EncodedNibblePairRight = ReadFromByteStream<uint8>(EncodedADPCMBlockRight, ReadIndexRight);
+
+			EncodedNibbleLeft = (EncodedNibblePairLeft >> 4) & 0x0F;
+			DecodedPCMData[WriteIndex++] = DecodeNibble(ContextLeft, EncodedNibbleLeft);
+
+			EncodedNibbleRight = (EncodedNibblePairRight >> 4) & 0x0F;
+			DecodedPCMData[WriteIndex++] = DecodeNibble(ContextRight, EncodedNibbleRight);
+
+			EncodedNibbleLeft = EncodedNibblePairLeft & 0x0F;
+			DecodedPCMData[WriteIndex++] = DecodeNibble(ContextLeft, EncodedNibbleLeft);
+
+			EncodedNibbleRight = EncodedNibblePairRight & 0x0F;
+			DecodedPCMData[WriteIndex++] = DecodeNibble(ContextRight, EncodedNibbleRight);
+		}
+
+	}
+} // end namespace ADPCM
+
+class ADPCMAUDIODECODER_API FAdpcmAudioDecoderModule : public IModuleInterface
+{
+public:	
+	TUniquePtr<IAudioInfoFactory> PcmFactory;
+	TUniquePtr<IAudioInfoFactory> AdpcmFactory;
+
+	virtual void StartupModule() override
+	{
+		AdpcmFactory = MakeUnique<FSimpleAudioInfoFactory>([] { return new FADPCMAudioInfo(); }, Audio::NAME_ADPCM);
+
+		// PCM is actually ADPCM at quality 100% 
+		PcmFactory = MakeUnique<FSimpleAudioInfoFactory>([] { return new FADPCMAudioInfo(); }, Audio::NAME_PCM);
+	}
+
+	virtual void ShutdownModule() override {}
+};
+
+IMPLEMENT_MODULE(FAdpcmAudioDecoderModule, AdpcmAudioDecoder)
diff --git a/Engine/Source/Runtime/AdpcmAudioDecoder/Module/Public/Decoders/ADPCMAudioInfo.h b/Engine/Source/Runtime/AdpcmAudioDecoder/Module/Public/Decoders/ADPCMAudioInfo.h
--- a/Engine/Source/Runtime/AdpcmAudioDecoder/Module/Public/Decoders/ADPCMAudioInfo.h
+++ b/Engine/Source/Runtime/AdpcmAudioDecoder/Module/Public/Decoders/ADPCMAudioInfo.h
@@ -91,15 +91,14 @@ namespace ADPCM
 class FADPCMAudioInfo : public ICompressedAudioInfo
 {
 public:
-	ENGINE_API FADPCMAudioInfo(void);
-	ENGINE_API virtual ~FADPCMAudioInfo(void);
-
+	ADPCMAUDIODECODER_API FADPCMAudioInfo(void);
+	ADPCMAUDIODECODER_API virtual ~FADPCMAudioInfo(void);
 	// ICompressedAudioInfo Interface
-	ENGINE_API virtual bool ReadCompressedInfo(const uint8* InSrcBufferData, uint32 InSrcBufferDataSize, struct FSoundQualityInfo* QualityInfo);
-	ENGINE_API virtual bool ReadCompressedData(uint8* Destination, bool bLooping, uint32 BufferSize);
-	ENGINE_API virtual void SeekToTime(const float SeekTime);
-	ENGINE_API virtual void SeekToFrame(const uint32 SeekFrame);
-	ENGINE_API virtual void ExpandFile(uint8* DstBuffer, struct FSoundQualityInfo* QualityInfo);
+	ADPCMAUDIODECODER_API virtual bool ReadCompressedInfo(const uint8* InSrcBufferData, uint32 InSrcBufferDataSize, struct FSoundQualityInfo* QualityInfo);
+	ADPCMAUDIODECODER_API virtual bool ReadCompressedData(uint8* Destination, bool bLooping, uint32 BufferSize);
+	ADPCMAUDIODECODER_API virtual void SeekToTime(const float SeekTime);
+	ADPCMAUDIODECODER_API virtual void SeekToFrame(const uint32 SeekFrame);
+	ADPCMAUDIODECODER_API virtual void ExpandFile(uint8* DstBuffer, struct FSoundQualityInfo* QualityInfo);
 	virtual void EnableHalfRate(bool HalfRate) {};
 	virtual uint32 GetSourceBufferSize() const { return SrcBufferDataSize; }
 	virtual bool UsesVorbisChannelOrdering() const { return false; }
@@ -186,4 +185,4 @@ private:
 	std::atomic<uint32> TargetSeekFrame;
 
 	FCriticalSection StreamSeekCriticalSection;
-};
+};
\ No newline at end of file
diff --git a/Engine/Source/Runtime/AudioMixer/Private/FileDecoder.cpp b/Engine/Source/Runtime/AudioMixer/Private/FileDecoder.cpp
--- a/Engine/Source/Runtime/AudioMixer/Private/FileDecoder.cpp
+++ b/Engine/Source/Runtime/AudioMixer/Private/FileDecoder.cpp
@@ -1,8 +1,6 @@
 // Copyright Epic Games, Inc. All Rights Reserved.
 
 #include "FileDecoder.h"
-#include "OpusAudioInfo.h"
-#include "VorbisAudioInfo.h"
 #include "HAL/PlatformFileManager.h"
 #include "DSP/FloatArrayMath.h"
 
@@ -60,38 +58,21 @@ bool FAudioFileReader::PopAudio(float* OutAudio, int32 NumSamples)
 }
 
 ICompressedAudioInfo* FAudioFileReader::GetNewDecompressorForFile(const FString& InPath)
-{
-	FString Extension = GetExtensionForFile(InPath);
-
-	static const FString OpusExtension = TEXT("opus");
-	static const FString OggExtension = TEXT("ogg");
+{	
+	using namespace Audio;
+	using FMapping = TTuple<FString, FName>; 
+	const FMapping Extensions[] =
+	{	
+		{ TEXT(".opus"), NAME_OPUS },
+		{ TEXT(".vorbis"), NAME_OGG },
+		{ TEXT(".binka"), NAME_BINKA }
+	};
 
-#if PLATFORM_SUPPORTS_OPUS_CODEC 
-	if (Extension.Equals(OpusExtension))
-	{
-		return new FOpusAudioInfo();
-	}
-#endif
-#if PLATFORM_SUPPORTS_VORBIS_CODEC 
-	if (Extension.Equals(OggExtension))
+	const FString LowerPath = InPath.ToLower();
+	if (const FMapping* Found = Algo::FindByPredicate(Extensions, [LowerPath](const auto &i) -> bool { return LowerPath.EndsWith(i.Key); }) )
 	{
-		return new FVorbisAudioInfo();
+		return IAudioInfoFactoryRegistry::Get().Create(Found->Value);		
 	}
-#endif
-
-	UE_LOG(LogTemp, Error, TEXT("Unknown extension '%s' for the FAudioFileReader formats supported on this platform."), *Extension);
+	UE_LOG(LogTemp, Error, TEXT("Unable to determin/create decompressor for '%s'"), *InPath);
 	return nullptr;
 }
-
-FString FAudioFileReader::GetExtensionForFile(const FString& InPath)
-{
-	int32 Index = INDEX_NONE;
-	if (InPath.FindLastChar(TCHAR('.'), Index))
-	{
-		return InPath.RightChop(Index + 1);
-	}
-	else
-	{
-		return FString();
-	}
-}
diff --git a/Engine/Source/Runtime/AudioMixer/Public/FileDecoder.h b/Engine/Source/Runtime/AudioMixer/Public/FileDecoder.h
--- a/Engine/Source/Runtime/AudioMixer/Public/FileDecoder.h
+++ b/Engine/Source/Runtime/AudioMixer/Public/FileDecoder.h
@@ -39,5 +39,4 @@ private:
 
 	ICompressedAudioInfo* GetNewDecompressorForFile(const FString& InPath);
 
-	FString GetExtensionForFile(const FString& InPath);
 };
diff --git a/Engine/Source/Runtime/Engine/Engine.Build.cs b/Engine/Source/Runtime/Engine/Engine.Build.cs
--- a/Engine/Source/Runtime/Engine/Engine.Build.cs
+++ b/Engine/Source/Runtime/Engine/Engine.Build.cs
@@ -171,14 +171,11 @@ public class Engine : ModuleRules
 				);
 		}
 
-		// Cross platform Audio Codecs:
-		AddEngineThirdPartyPrivateStaticDependencies(Target,
-			"UEOgg",
-			"Vorbis",
-			"VorbisFile",
-			"libOpus"
-			);
+		// Cross platform Audio Codecs: (we build here, but don't depend on them directly)
 		DynamicallyLoadedModuleNames.Add("BinkAudioDecoder");
+		DynamicallyLoadedModuleNames.Add("VorbisAudioDecoder");
+		DynamicallyLoadedModuleNames.Add("OpusAudioDecoder");
+		DynamicallyLoadedModuleNames.Add("AdpcmAudioDecoder");
 		
 		DynamicallyLoadedModuleNames.Add("EyeTracker");
 
@@ -449,12 +446,6 @@ public class Engine : ModuleRules
 
 		if (Target.IsInPlatformGroup(UnrealPlatformGroup.Android))
 		{
-			AddEngineThirdPartyPrivateStaticDependencies(Target,
-				"UEOgg",
-				"Vorbis",
-				"VorbisFile"
-				);
-
 			PrivateIncludePathModuleNames.Add("AndroidRuntimeSettings");
 		}
 
@@ -464,16 +455,6 @@ public class Engine : ModuleRules
 			PrivateIncludePathModuleNames.Add("IOSRuntimeSettings");
 		}
 
-		if (Target.IsInPlatformGroup(UnrealPlatformGroup.Unix))
-		{
-			AddEngineThirdPartyPrivateStaticDependencies(Target,
-				"UEOgg",
-				"Vorbis",
-				"VorbisFile",
-				"libOpus"
-				);
-		}
-
 		PublicDefinitions.Add("GPUPARTICLE_LOCAL_VF_ONLY=0");
 
 		// Add a reference to the stats HTML files referenced by UEngine::DumpFPSChartToHTML. Previously staged by CopyBuildToStagingDirectory.
diff --git a/Engine/Source/Runtime/Engine/Private/ADPCMAudioInfo.cpp b/Engine/Source/Runtime/Engine/Private/ADPCMAudioInfo.cpp
deleted file mode 100644
diff --git a/Engine/Source/Runtime/Engine/Private/Audio.cpp b/Engine/Source/Runtime/Engine/Private/Audio.cpp
--- a/Engine/Source/Runtime/Engine/Private/Audio.cpp
+++ b/Engine/Source/Runtime/Engine/Private/Audio.cpp
@@ -42,8 +42,6 @@ DEFINE_STAT(STAT_AudioFinishedDelegates);
 DEFINE_STAT(STAT_AudioBufferTime);
 DEFINE_STAT(STAT_AudioBufferTimeChannels);
 
-DEFINE_STAT(STAT_VorbisDecompressTime);
-DEFINE_STAT(STAT_VorbisPrepareDecompressionTime);
 DEFINE_STAT(STAT_AudioDecompressTime);
 DEFINE_STAT(STAT_AudioPrepareDecompressionTime);
 DEFINE_STAT(STAT_AudioStreamedDecompressTime);
diff --git a/Engine/Source/Runtime/Engine/Private/AudioDecompress.cpp b/Engine/Source/Runtime/Engine/Private/AudioDecompress.cpp
--- a/Engine/Source/Runtime/Engine/Private/AudioDecompress.cpp
+++ b/Engine/Source/Runtime/Engine/Private/AudioDecompress.cpp
@@ -5,7 +5,6 @@
 #include "AudioDevice.h"
 #include "Engine/Engine.h"
 #include "Interfaces/IAudioFormat.h"
-#include "ADPCMAudioInfo.h"
 #include "AudioStreamingCache.h"
 #include "Misc/CoreStats.h"
 #include "Misc/ScopeRWLock.h"
@@ -688,259 +687,6 @@ const FStreamedAudioChunkSeekTable& IStreamedCompressedInfo::GetCurrentSeekTable
 	return const_cast<IStreamedCompressedInfo*>(this)->GetCurrentSeekTable();
 }
 
-//////////////////////////////////////////////////////////////////////////
-// Copied from IOS - probably want to split and share
-//////////////////////////////////////////////////////////////////////////
-#define NUM_ADAPTATION_TABLE 16
-#define NUM_ADAPTATION_COEFF 7
-#define SOUND_SOURCE_FREE 0
-#define SOUND_SOURCE_LOCKED 1
-
-namespace
-{
-	template <typename T, uint32 B>
-	inline T SignExtend(const T ValueToExtend)
-	{
-		struct { T ExtendedValue : B; } SignExtender;
-		return SignExtender.ExtendedValue = ValueToExtend;
-	}
-
-	template <typename T>
-	inline T ReadFromByteStream(const uint8* ByteStream, int32& ReadIndex, bool bLittleEndian = true)
-	{
-		T ValueRaw = 0;
-
-		if (bLittleEndian)
-		{
-#if PLATFORM_LITTLE_ENDIAN
-			for (int32 ByteIndex = 0; ByteIndex < sizeof(T); ++ByteIndex)
-#else
-			for (int32 ByteIndex = sizeof(T) - 1; ByteIndex >= 0; --ByteIndex)
-#endif // PLATFORM_LITTLE_ENDIAN
-			{
-				ValueRaw |= ByteStream[ReadIndex++] << 8 * ByteIndex;
-			}
-		}
-		else
-		{
-#if PLATFORM_LITTLE_ENDIAN
-			for (int32 ByteIndex = sizeof(T) - 1; ByteIndex >= 0; --ByteIndex)
-#else
-			for (int32 ByteIndex = 0; ByteIndex < sizeof(T); ++ByteIndex)
-#endif // PLATFORM_LITTLE_ENDIAN
-			{
-				ValueRaw |= ByteStream[ReadIndex++] << 8 * ByteIndex;
-			}
-		}
-
-		return ValueRaw;
-	}
-
-	template <typename T>
-	inline void WriteToByteStream(T Value, uint8* ByteStream, int32& WriteIndex, bool bLittleEndian = true)
-	{
-		if (bLittleEndian)
-		{
-#if PLATFORM_LITTLE_ENDIAN
-			for (int32 ByteIndex = 0; ByteIndex < sizeof(T); ++ByteIndex)
-#else
-			for (int32 ByteIndex = sizeof(T) - 1; ByteIndex >= 0; --ByteIndex)
-#endif // PLATFORM_LITTLE_ENDIAN
-			{
-				ByteStream[WriteIndex++] = (Value >> (8 * ByteIndex)) & 0xFF;
-			}
-		}
-		else
-		{
-#if PLATFORM_LITTLE_ENDIAN
-			for (int32 ByteIndex = sizeof(T) - 1; ByteIndex >= 0; --ByteIndex)
-#else
-			for (int32 ByteIndex = 0; ByteIndex < sizeof(T); ++ByteIndex)
-#endif // PLATFORM_LITTLE_ENDIAN
-			{
-				ByteStream[WriteIndex++] = (Value >> (8 * ByteIndex)) & 0xFF;
-			}
-		}
-	}
-
-	template <typename T>
-	inline T ReadFromArray(const T* ElementArray, int32& ReadIndex, int32 NumElements, int32 IndexStride = 1)
-	{
-		T OutputValue = 0;
-
-		if (ReadIndex >= 0 && ReadIndex < NumElements)
-		{
-			OutputValue = ElementArray[ReadIndex];
-			ReadIndex += IndexStride;
-		}
-
-		return OutputValue;
-	}
-
-	inline bool LockSourceChannel(int32* ChannelLock)
-	{
-		check(ChannelLock != NULL);
-		return FPlatformAtomics::InterlockedCompareExchange(ChannelLock, SOUND_SOURCE_LOCKED, SOUND_SOURCE_FREE) == SOUND_SOURCE_FREE;
-	}
-
-	inline void UnlockSourceChannel(int32* ChannelLock)
-	{
-		check(ChannelLock != NULL);
-		int32 Result = FPlatformAtomics::InterlockedCompareExchange(ChannelLock, SOUND_SOURCE_FREE, SOUND_SOURCE_LOCKED);
-
-		check(Result == SOUND_SOURCE_LOCKED);
-	}
-
-} // end namespace
-
-namespace ADPCM
-{
-	template <typename T>
-	static void GetAdaptationTable(T(&OutAdaptationTable)[NUM_ADAPTATION_TABLE])
-	{
-		// Magic values as specified by standard
-		static T AdaptationTable[] =
-		{
-			230, 230, 230, 230, 307, 409, 512, 614,
-			768, 614, 512, 409, 307, 230, 230, 230
-		};
-
-		FMemory::Memcpy(&OutAdaptationTable, AdaptationTable, sizeof(AdaptationTable));
-	}
-
-	struct FAdaptationContext
-	{
-	public:
-		// Adaptation constants
-		int32 AdaptationTable[NUM_ADAPTATION_TABLE];
-		int32 AdaptationCoefficient1[NUM_ADAPTATION_COEFF];
-		int32 AdaptationCoefficient2[NUM_ADAPTATION_COEFF];
-
-		int32 AdaptationDelta;
-		int32 Coefficient1;
-		int32 Coefficient2;
-		int32 Sample1;
-		int32 Sample2;
-
-		FAdaptationContext() :
-			AdaptationDelta(0),
-			Coefficient1(0),
-			Coefficient2(0),
-			Sample1(0),
-			Sample2(0)
-		{
-			GetAdaptationTable(AdaptationTable);
-			GetAdaptationCoefficients(AdaptationCoefficient1, AdaptationCoefficient2);
-		}
-	};
-
-	FORCEINLINE int16 DecodeNibble(FAdaptationContext& Context, uint8 EncodedNibble)
-	{
-		int32 PredictedSample = (Context.Sample1 * Context.Coefficient1 + Context.Sample2 * Context.Coefficient2) / 256;
-		PredictedSample += SignExtend<int8, 4>(EncodedNibble) * Context.AdaptationDelta;
-		PredictedSample = FMath::Clamp(PredictedSample, -32768, 32767);
-
-		// Shuffle samples for the next iteration
-		Context.Sample2 = Context.Sample1;
-		Context.Sample1 = static_cast<int16>(PredictedSample);
-		Context.AdaptationDelta = (Context.AdaptationDelta * Context.AdaptationTable[EncodedNibble]) / 256;
-		Context.AdaptationDelta = FMath::Max(Context.AdaptationDelta, 16);
-
-		return Context.Sample1;
-	}
-
-	void DecodeBlock(const uint8* EncodedADPCMBlock, int32 BlockSize, int16* DecodedPCMData)
-	{
-		FAdaptationContext Context;
-		int32 ReadIndex = 0;
-		int32 WriteIndex = 0;
-
-		uint8 CoefficientIndex = ReadFromByteStream<uint8>(EncodedADPCMBlock, ReadIndex);
-		Context.AdaptationDelta = ReadFromByteStream<int16>(EncodedADPCMBlock, ReadIndex);
-		Context.Sample1 = ReadFromByteStream<int16>(EncodedADPCMBlock, ReadIndex);
-		Context.Sample2 = ReadFromByteStream<int16>(EncodedADPCMBlock, ReadIndex);
-		Context.Coefficient1 = Context.AdaptationCoefficient1[CoefficientIndex];
-		Context.Coefficient2 = Context.AdaptationCoefficient2[CoefficientIndex];
-
-		// The first two samples are sent directly to the output in reverse order, as per the standard
-		DecodedPCMData[WriteIndex++] = Context.Sample2;
-		DecodedPCMData[WriteIndex++] = Context.Sample1;
-
-		uint8 EncodedNibblePair = 0;
-		uint8 EncodedNibble = 0;
-		while (ReadIndex < BlockSize)
-		{
-			// Read from the byte stream and advance the read head.
-			EncodedNibblePair = ReadFromByteStream<uint8>(EncodedADPCMBlock, ReadIndex);
-
-			EncodedNibble = (EncodedNibblePair >> 4) & 0x0F;
-			DecodedPCMData[WriteIndex++] = DecodeNibble(Context, EncodedNibble);
-
-			EncodedNibble = EncodedNibblePair & 0x0F;
-			DecodedPCMData[WriteIndex++] = DecodeNibble(Context, EncodedNibble);
-		}
-	}
-
-	// Decode two PCM streams and interleave as stereo data
-	void DecodeBlockStereo(const uint8* EncodedADPCMBlockLeft, const uint8* EncodedADPCMBlockRight, int32 BlockSize, int16* DecodedPCMData)
-	{
-		FAdaptationContext ContextLeft;
-		FAdaptationContext ContextRight;
-		int32 ReadIndexLeft = 0;
-		int32 ReadIndexRight = 0;
-		int32 WriteIndex = 0;
-
-		uint8 CoefficientIndexLeft = ReadFromByteStream<uint8>(EncodedADPCMBlockLeft, ReadIndexLeft);
-		ContextLeft.AdaptationDelta = ReadFromByteStream<int16>(EncodedADPCMBlockLeft, ReadIndexLeft);
-		ContextLeft.Sample1 = ReadFromByteStream<int16>(EncodedADPCMBlockLeft, ReadIndexLeft);
-		ContextLeft.Sample2 = ReadFromByteStream<int16>(EncodedADPCMBlockLeft, ReadIndexLeft);
-		ContextLeft.Coefficient1 = ContextLeft.AdaptationCoefficient1[CoefficientIndexLeft];
-		ContextLeft.Coefficient2 = ContextLeft.AdaptationCoefficient2[CoefficientIndexLeft];
-
-		uint8 CoefficientIndexRight = ReadFromByteStream<uint8>(EncodedADPCMBlockRight, ReadIndexRight);
-		ContextRight.AdaptationDelta = ReadFromByteStream<int16>(EncodedADPCMBlockRight, ReadIndexRight);
-		ContextRight.Sample1 = ReadFromByteStream<int16>(EncodedADPCMBlockRight, ReadIndexRight);
-		ContextRight.Sample2 = ReadFromByteStream<int16>(EncodedADPCMBlockRight, ReadIndexRight);
-		ContextRight.Coefficient1 = ContextRight.AdaptationCoefficient1[CoefficientIndexRight];
-		ContextRight.Coefficient2 = ContextRight.AdaptationCoefficient2[CoefficientIndexRight];
-
-		// The first two samples from each stream are sent directly to the output in reverse order, as per the standard
-		DecodedPCMData[WriteIndex++] = ContextLeft.Sample2;
-		DecodedPCMData[WriteIndex++] = ContextRight.Sample2;
-		DecodedPCMData[WriteIndex++] = ContextLeft.Sample1;
-		DecodedPCMData[WriteIndex++] = ContextRight.Sample1;
-
-		uint8 EncodedNibblePairLeft = 0;
-		uint8 EncodedNibbleLeft = 0;
-		uint8 EncodedNibblePairRight = 0;
-		uint8 EncodedNibbleRight = 0;
-
-		while (ReadIndexLeft < BlockSize)
-		{
-			// Read from the byte stream and advance the read head.
-			EncodedNibblePairLeft = ReadFromByteStream<uint8>(EncodedADPCMBlockLeft, ReadIndexLeft);
-			EncodedNibblePairRight = ReadFromByteStream<uint8>(EncodedADPCMBlockRight, ReadIndexRight);
-
-			EncodedNibbleLeft = (EncodedNibblePairLeft >> 4) & 0x0F;
-			DecodedPCMData[WriteIndex++] = DecodeNibble(ContextLeft, EncodedNibbleLeft);
-
-			EncodedNibbleRight = (EncodedNibblePairRight >> 4) & 0x0F;
-			DecodedPCMData[WriteIndex++] = DecodeNibble(ContextRight, EncodedNibbleRight);
-
-			EncodedNibbleLeft = EncodedNibblePairLeft & 0x0F;
-			DecodedPCMData[WriteIndex++] = DecodeNibble(ContextLeft, EncodedNibbleLeft);
-
-			EncodedNibbleRight = EncodedNibblePairRight & 0x0F;
-			DecodedPCMData[WriteIndex++] = DecodeNibble(ContextRight, EncodedNibbleRight);
-		}
-
-	}
-} // end namespace ADPCM
-
-//////////////////////////////////////////////////////////////////////////
-// End of copied section
-//////////////////////////////////////////////////////////////////////////
-
 /**
  * Worker for decompression on a separate thread
  */
diff --git a/Engine/Source/Runtime/Engine/Private/AudioDevice.cpp b/Engine/Source/Runtime/Engine/Private/AudioDevice.cpp
--- a/Engine/Source/Runtime/Engine/Private/AudioDevice.cpp
+++ b/Engine/Source/Runtime/Engine/Private/AudioDevice.cpp
@@ -2,8 +2,8 @@
 #include "AudioDevice.h"
 
 #include "ActiveSoundUpdateInterface.h"
-#include "ADPCMAudioInfo.h"
 #include "AudioCompressionSettingsUtils.h"
+#include "AudioDecompress.h"
 #include "AudioEffect.h"
 #include "AudioMixerTrace.h"
 #include "AudioPluginUtilities.h"
@@ -239,14 +239,10 @@ namespace Audio
 		{
 			return nullptr;
 		}
-		// Multi-platform decoders -- can be instantiated here for all platforms.
-		else if (InRuntimeFormat == Audio::NAME_ADPCM || InRuntimeFormat == Audio::NAME_PCM)
-		{
-			// Our PCM is currently really our "ADPCM" encoder at 100 quality
-			return new FADPCMAudioInfo();
-		}
-		ensureMsgf(false, TEXT("Unknown runtime sound asset format type."));
-		return nullptr;
+
+		ICompressedAudioInfo* Decoder = IAudioInfoFactoryRegistry::Get().Create(InRuntimeFormat);
+		ensureMsgf(Decoder, TEXT("Unknown runtime sound asset format type. '%s' "), *InRuntimeFormat.ToString());
+		return Decoder;
 	}
 
 	uint64 GetTransmitterID(uint64 ComponentID, UPTRINT WaveInstanceHash, uint32 PlayOrder)
diff --git a/Engine/Source/Runtime/Engine/Private/AudioDeviceManager.cpp b/Engine/Source/Runtime/Engine/Private/AudioDeviceManager.cpp
--- a/Engine/Source/Runtime/Engine/Private/AudioDeviceManager.cpp
+++ b/Engine/Source/Runtime/Engine/Private/AudioDeviceManager.cpp
@@ -15,9 +15,7 @@
 #include "UObject/UObjectIterator.h"
 
 #include "Sound/AudioFormatSettings.h"
-#include "ADPCMAudioInfo.h"
-#include "VorbisAudioInfo.h"
-#include "OpusAudioInfo.h"
+#include "AudioDecompress.h"
 
 #if INSTRUMENT_AUDIODEVICE_HANDLES
 #include "Containers/StringConv.h"
@@ -356,17 +354,7 @@ void FAudioDeviceManager::RegisterAudioInfoFactories()
 			FModuleManager::Get().LoadModuleChecked(*i);
 		}
 	}
-	
-	// Hardcode the loading of the built in codecs.
-	LoadVorbisLibraries();
-	FModuleManager::Get().LoadModuleChecked(TEXT("BinkAudioDecoder"));
-
-	// Register the engine formats.
-	EngineFormats.Add(MakePimpl<FSimpleAudioInfoFactory>([] { return new FADPCMAudioInfo(); }, Audio::NAME_PCM));
-	EngineFormats.Add(MakePimpl<FSimpleAudioInfoFactory>([] { return new FADPCMAudioInfo(); }, Audio::NAME_ADPCM));
-	EngineFormats.Add(MakePimpl<FSimpleAudioInfoFactory>([] { return new FVorbisAudioInfo(); }, Audio::NAME_OGG));
-	EngineFormats.Add(MakePimpl<FSimpleAudioInfoFactory>([] { return new FOpusAudioInfo(); }, Audio::NAME_OPUS));
-
+		
 	// Sanity check we have all the Factories we need to run now by 
 	TArray<FName> AllFormats;
 	GetAudioFormatSettings().GetAllWaveFormats(AllFormats);
diff --git a/Engine/Source/Runtime/Engine/Private/OpusAudioInfo.cpp b/Engine/Source/Runtime/Engine/Private/OpusAudioInfo.cpp
deleted file mode 100644
diff --git a/Engine/Source/Runtime/Engine/Private/VorbisAudioInfo.cpp b/Engine/Source/Runtime/Engine/Private/VorbisAudioInfo.cpp
deleted file mode 100644
diff --git a/Engine/Source/Runtime/Engine/Public/ADPCMAudioInfo.h b/Engine/Source/Runtime/Engine/Public/ADPCMAudioInfo.h
--- a/Engine/Source/Runtime/Engine/Public/ADPCMAudioInfo.h
+++ b/Engine/Source/Runtime/Engine/Public/ADPCMAudioInfo.h
@@ -1,189 +1,15 @@
 // Copyright Epic Games, Inc. All Rights Reserved.
 
-/*=============================================================================
-	ADPCMAudioInfo.h: Unreal audio ADPCM decompression interface object.
-=============================================================================*/
-
+/**
+ * DEPRECATED 5.4 ADPCMAudioInfo.h. please use the module path of Decoders/ADPCMAudioInfo.h instead.
+ */
 #pragma once
 
-#include "CoreMinimal.h"
-#include "Audio.h"
-#include "AudioDecompress.h"
-#include "Sound/SoundWave.h"
-#include "ContentStreaming.h"
-
-#define NUM_ADAPTATION_TABLE 16
-#define NUM_ADAPTATION_COEFF 7
-
-#define WAVE_FORMAT_LPCM  1
-#ifndef WAVE_FORMAT_ADPCM
-#define WAVE_FORMAT_ADPCM 2
-#endif
-
-#if PLATFORM_SUPPORTS_PRAGMA_PACK
-#pragma pack(push, 2)
-#endif
-	struct WaveFormatHeader
-	{
-		uint16 wFormatTag;      // Format type: 1 = PCM, 2 = ADPCM
-		uint16 nChannels;       // Number of channels (i.e. mono, stereo...).
-		uint32 nSamplesPerSec;  // Sample rate. 44100 or 22050 or 11025  Hz.
-		uint32 nAvgBytesPerSec; // For buffer estimation  = sample rate * BlockAlign.
-		uint16 nBlockAlign;     // Block size of data = Channels times BYTES per sample.
-		uint16 wBitsPerSample;  // Number of bits per sample of mono data.
-		uint16 cbSize;          // The count in bytes of the size of extra information (after cbSize).
-	};
-#if PLATFORM_SUPPORTS_PRAGMA_PACK
-#pragma pack(pop)
+#ifdef _MSC_VER
+	#pragma message(__FILE__"(9): has moved to it's own module ADPCMAudioDecoder, please update your code.")
+#else
+	#pragma message("ADPCMAudioInfo.h has moved to it's own module ADPCMAudioDecoder, please update your code.")
 #endif
 
-namespace ADPCM
-{
-
-	template <typename T>
-	static void GetAdaptationCoefficients(T(& OutAdaptationCoefficient1)[NUM_ADAPTATION_COEFF], T(& OutAdaptationCoefficient2)[NUM_ADAPTATION_COEFF])
-	{
-		// Magic values as specified by standard
-		static T AdaptationCoefficient1[] =
-		{
-			256, 512, 0, 192, 240, 460, 392
-		};
-		static T AdaptationCoefficient2[] =
-		{
-			0, -256, 0,  64, 0, -208, -232
-		};
-
-		FMemory::Memcpy(&OutAdaptationCoefficient1, AdaptationCoefficient1, sizeof(AdaptationCoefficient1));
-		FMemory::Memcpy(&OutAdaptationCoefficient2, AdaptationCoefficient2, sizeof(AdaptationCoefficient2));
-	}
-
-#if PLATFORM_SUPPORTS_PRAGMA_PACK
-#pragma pack(push, 2)
-#endif
-	struct ADPCMFormatHeader
-	{
-		WaveFormatHeader	BaseFormat;
-		uint16				wSamplesPerBlock;
-		uint16				wNumCoef;
-		int16				aCoef[2 * NUM_ADAPTATION_COEFF];
-		uint32				SamplesPerChannel;	// This is the exact samples per channel for sample precise looping
-
-		ADPCMFormatHeader()
-		{
-			int16 AdaptationCoefficient1[NUM_ADAPTATION_COEFF];
-			int16 AdaptationCoefficient2[NUM_ADAPTATION_COEFF];
-			GetAdaptationCoefficients(AdaptationCoefficient1, AdaptationCoefficient2);
-
-			// Interlace the coefficients as pairs
-			for (int32 Coeff = 0, WriteIndex = 0; Coeff < NUM_ADAPTATION_COEFF; Coeff++)
-			{
-				aCoef[WriteIndex++] = AdaptationCoefficient1[Coeff];
-				aCoef[WriteIndex++] = AdaptationCoefficient2[Coeff];
-			}
-		}
-	};
-#if PLATFORM_SUPPORTS_PRAGMA_PACK
-#pragma pack(pop)
-#endif
-
-};
-
-class FADPCMAudioInfo : public ICompressedAudioInfo
-{
-public:
-	ENGINE_API FADPCMAudioInfo(void);
-	ENGINE_API virtual ~FADPCMAudioInfo(void);
-
-	// ICompressedAudioInfo Interface
-	ENGINE_API virtual bool ReadCompressedInfo(const uint8* InSrcBufferData, uint32 InSrcBufferDataSize, struct FSoundQualityInfo* QualityInfo);
-	ENGINE_API virtual bool ReadCompressedData(uint8* Destination, bool bLooping, uint32 BufferSize);
-	ENGINE_API virtual void SeekToTime(const float SeekTime);
-	ENGINE_API virtual void SeekToFrame(const uint32 SeekFrame);
-	ENGINE_API virtual void ExpandFile(uint8* DstBuffer, struct FSoundQualityInfo* QualityInfo);
-	virtual void EnableHalfRate(bool HalfRate) {};
-	virtual uint32 GetSourceBufferSize() const { return SrcBufferDataSize; }
-	virtual bool UsesVorbisChannelOrdering() const { return false; }
-	virtual int GetStreamBufferSize() const;
-
-	// Additional overrides for streaming
-	virtual bool SupportsStreaming() const override {return true;}
-	virtual bool StreamCompressedInfoInternal(const FSoundWaveProxyPtr& InWaveProxy, struct FSoundQualityInfo* QualityInfo) override;
-	virtual bool StreamCompressedData(uint8* Destination, bool bLooping, uint32 BufferSize, int32& OutNumBytesStreamed) override;
-	virtual int32 GetCurrentChunkIndex() const override
-	{
-		return CurrentChunkIndex;
-	}
-	virtual int32 GetCurrentChunkOffset() const override {return CurrentChunkBufferOffset;}
-	virtual bool ReleaseStreamChunk(bool bBlockUntilReleased) override;
-	// End of ICompressedAudioInfo Interface
-
-	uint16 GetFormatTag()
-	{
-		if (WaveInfo.pFormatTag)
-		{
-			return *WaveInfo.pFormatTag;
-		}
-		else
-		{
-			return 0;
-		}
-	}
-
-	int32 GetNumChannels()
-	{
-		return NumChannels;
-	}
-
-private:
-
-	// Wrapper function that returns a pointer to the currently used compressed data.
-	// If a non-zero chunk is requested, this function also aquires a reference to that chunk
-	// until we move on to a different chunk.
-	const uint8* GetLoadedChunk(const FSoundWaveProxyPtr& InSoundWave, uint32 ChunkIndex, uint32& OutChunkSize);
-
-	int32 NumConsecutiveReadFailiures;
-
-	FWaveModInfo WaveInfo;
-	const uint8*	SrcBufferData;
-	uint32			SrcBufferDataSize;
-
-	uint32	UncompressedBlockSize;
-	uint32	CompressedBlockSize;
-	uint32	BlockSize;
-	int32	StreamBufferSize;
-	uint32	TotalDecodedSize;
-	int32	NumChannels;
-	int32	Format;
-
-	uint32 PreviouslyRequestedChunkIndex;
-
-	uint8*			UncompressedBlockData;			// This holds the current block of compressed data for all channels
-	uint32			CurrentUncompressedBlockSampleIndex;	// This is the sample index within the current uncompressed block data
-	uint32			CurrentChunkIndex;				// This is the index that is currently being used, needed by streaming engine to make sure it stays loaded and the next chunk gets preloaded
-	uint32			CurrentChunkBufferOffset;		// This is this byte offset within the current chunk, used by streaming engine to prioritize a load if more then half way through current chunk
-	uint32			CurrentChunkDataSize;			// The size of the current chunk, the first chunk is bigger to accommodate the header info
-	uint32			TotalSamplesStreamed;			// The number of samples streamed so far (per channel)
-	uint32			TotalSamplesPerChannel;			// Number of samples per channel, used to detect when an audio waveform has ended
-	uint32			SamplesPerBlock;				// The number of samples per block
-	uint32			FirstChunkSampleDataOffset;		// The size of the header in the first chunk, used to skip over it when looping or starting the sample over
-	uint32          FirstChunkSampleDataIndex;
-	const uint8*	CurCompressedChunkData;			// A pointer to the current chunk of data
-	FAudioChunkHandle CurCompressedChunkHandle;     // Shared reference to the current chunk of data.
-	FCriticalSection CurCompressedChunkHandleCriticalSection;
-	bool bDecompressorReleased;                     // When the chunk is released, we raise this to true, then early exit on future decodes.
-
-	uint32			CurrentCompressedBlockIndex;		// For non disk streaming - the current compressed block in the compressed source data
-	uint32			TotalCompressedBlocksPerChannel;	// For non disk streaming - the total number of compressed blocks per channel
-	uint8			bNewSeekRequest : 1;				// Whether or not there is a new seek request to be processed 
-	uint8			bSeekPendingRead : 1;				// Whether or not a requested seek is pending a read
-	uint8           bSeekedFowardToNextChunk : 1;       // If this is true, we have already seeked forward to the next chunk while waiting for the current chunk of audio to load.
-
-private:
-	void ProcessSeekRequest();
-	void SeekToFrameInternal(const uint32 InSeekFrame);
-	void ResetSeekState();
-
-	std::atomic<uint32> TargetSeekFrame;
-
-	FCriticalSection StreamSeekCriticalSection;
-};
+ // Include the new file so that the project still compiles, but has warnings
+#include "Decoders/ADPCMAudioInfo.h"
diff --git a/Engine/Source/Runtime/Engine/Public/Audio.h b/Engine/Source/Runtime/Engine/Public/Audio.h
--- a/Engine/Source/Runtime/Engine/Public/Audio.h
+++ b/Engine/Source/Runtime/Engine/Public/Audio.h
@@ -49,9 +49,7 @@ DECLARE_CYCLE_STAT_EXTERN(TEXT("Source Init"), STAT_AudioSourceInitTime, STATGRO
 DECLARE_CYCLE_STAT_EXTERN(TEXT("Source Create"), STAT_AudioSourceCreateTime, STATGROUP_Audio, ENGINE_API);
 DECLARE_CYCLE_STAT_EXTERN(TEXT("Submit Buffers"), STAT_AudioSubmitBuffersTime, STATGROUP_Audio, ENGINE_API);
 DECLARE_CYCLE_STAT_EXTERN(TEXT("Decompress Audio"), STAT_AudioDecompressTime, STATGROUP_Audio, ENGINE_API);
-DECLARE_CYCLE_STAT_EXTERN(TEXT("Decompress Vorbis"), STAT_VorbisDecompressTime, STATGROUP_Audio, );
 DECLARE_CYCLE_STAT_EXTERN(TEXT("Prepare Audio Decompression"), STAT_AudioPrepareDecompressionTime, STATGROUP_Audio, );
-DECLARE_CYCLE_STAT_EXTERN(TEXT("Prepare Vorbis Decompression"), STAT_VorbisPrepareDecompressionTime, STATGROUP_Audio, );
 DECLARE_CYCLE_STAT_EXTERN(TEXT("Finding Nearest Location"), STAT_AudioFindNearestLocation, STATGROUP_Audio, );
 DECLARE_CYCLE_STAT_EXTERN(TEXT("Decompress Streamed"), STAT_AudioStreamedDecompressTime, STATGROUP_Audio, );
 DECLARE_CYCLE_STAT_EXTERN(TEXT("Buffer Creation"), STAT_AudioResourceCreationTime, STATGROUP_Audio, );
diff --git a/Engine/Source/Runtime/Engine/Public/OpusAudioInfo.h b/Engine/Source/Runtime/Engine/Public/OpusAudioInfo.h
--- a/Engine/Source/Runtime/Engine/Public/OpusAudioInfo.h
+++ b/Engine/Source/Runtime/Engine/Public/OpusAudioInfo.h
@@ -1,92 +1,15 @@
 // Copyright Epic Games, Inc. All Rights Reserved.
 
-/*=============================================================================
-	OpusAudioInfo.h: Unreal audio opus decompression interface object.
-=============================================================================*/
-
-#pragma once
-
-#include "CoreMinimal.h"
-#include "AudioDecompress.h"
-
-struct FOpusDecoderWrapper;
-struct FSoundQualityInfo;
-
 /**
-* Helper class to parse opus data
-*/
-class FOpusAudioInfo : public IStreamedCompressedInfo
-{
-public:
-	ENGINE_API FOpusAudioInfo();
-	ENGINE_API virtual ~FOpusAudioInfo();
-
-	//~ Begin IStreamedCompressedInfo Interface
-	bool ParseHeader(const uint8* InSrcBufferData, uint32 InSrcBufferDataSize, FSoundQualityInfo* QualityInfo) override;
-	int32 GetFrameSize() override;
-	uint32 GetMaxFrameSizeSamples() const override;
-	bool CreateDecoder() override;
-	FDecodeResult Decode(const uint8* CompressedData, const int32 CompressedDataSize, uint8* OutPCMData, const int32 OutputPCMDataSize) override;
-	void PrepareToLoop() override;
-	void SeekToTime(const float SeekTime) override;
-	void SeekToFrame(const uint32 SeekFrame) override;
-	//~ End IStreamedCompressedInfo Interface
-
-	struct ENGINE_API FHeader
-	{
-		char Identifier[8];
-		uint8 Version = 0;
-		uint8 NumChannels = 0;
-		uint32 SampleRate = 0;
-		uint32 EncodedSampleRate = 0;
-		uint64 ActiveSampleCount = 0;
-		uint32 NumEncodedFrames = 0;
-		int32 NumPreSkipSamples = 0;
-		int32 NumSilentSamplesAtBeginning = 0;
-		int32 NumSilentSamplesAtEnd = 0;
-
-		void Reset()
-		{
-			FMemory::Memzero(Identifier);
-			Version = 0;
-			NumChannels = 0;
-			SampleRate = 0;
-			EncodedSampleRate = 0;
-			ActiveSampleCount = 0;
-			NumEncodedFrames = 0;
-			NumPreSkipSamples = 0;
-			NumSilentSamplesAtBeginning = 0;
-			NumSilentSamplesAtEnd = 0;
-		}
-
-		static inline const char OPUS_ID[8] {'U','E','O','P','U','S','\0','\0'};
-		static constexpr int32 HeaderSize()
-		{
-			return
-				  sizeof(char)*8	// Identifier
-				+ sizeof(uint8)		// Version
-				+ sizeof(uint8)		// NumChannels
-				+ sizeof(uint32)	// SampleRate
-				+ sizeof(uint32)	// EncodedSampleRate
-				+ sizeof(uint64)	// ActiveSampleCount
-				+ sizeof(uint32)	// NumEncodedFrames
-				+ sizeof(int32)		// NumPreSkipSamples
-				+ sizeof(int32)		// NumSilentSamplesAtBeginning
-				+ sizeof(int32);	// NumSilentSamplesAtEnd
-		}
-	};
-	static ENGINE_API bool ParseHeader(FHeader& OutHeader, uint32& OutNumRead, const uint8* InSrcBufferData, uint32 InSrcBufferDataSize);
-
-protected:
-	/** Wrapper around Opus-specific decoding state and APIs */
-	FOpusDecoderWrapper* OpusDecoderWrapper;
+ * DEPRECATED 5.4 OpusAudioInfo.h. please use the module path of Decoders/OpusAudioInfo.h instead.
+ */
+#pragma once
 
-	/** The currently active header. */
-	FHeader Header;
-	
-	/** The number of decoded samples that need to be skipped over. */
-	int32 NumRemainingSamplesToSkip = 0;
+#ifdef _MSC_VER
+	#pragma message(__FILE__"(9): has moved to it's own module OpusAudioDecoder, please update your code.")
+#else
+	#pragma message("OpusAudioInfo.h has moved to it's own module OpusAudioDecoder, please update your code.")
+#endif
 
-	/** Samples that are left over from the previous Decode() call that did not fit into the receive buffer. */
-	TArray<uint8> PreviousDecodedUnusedSamples;
-};
+ // Include the new file so that the project still compiles, but has warnings
+#include "Decoders/OpusAudioInfo.h"
diff --git a/Engine/Source/Runtime/Engine/Public/VorbisAudioInfo.h b/Engine/Source/Runtime/Engine/Public/VorbisAudioInfo.h
--- a/Engine/Source/Runtime/Engine/Public/VorbisAudioInfo.h
+++ b/Engine/Source/Runtime/Engine/Public/VorbisAudioInfo.h
@@ -1,144 +1,15 @@
 // Copyright Epic Games, Inc. All Rights Reserved.
 
-/*=============================================================================
-	VorbisAudioInfo.h: Unreal audio vorbis decompression interface object.
-=============================================================================*/
-
-#pragma once
-
-#include "CoreMinimal.h"
-#include "HAL/ThreadSafeBool.h"
-#include "AudioDecompress.h"
-
-/**
- * Whether to use OggVorbis audio format.
- **/
-#ifndef WITH_OGGVORBIS
-	#define WITH_OGGVORBIS 0
-#endif
-
-/**
- * Whether to load OggVorbis DLLs
- **/
-#ifndef WITH_OGGVORBIS_DLL
-	#define WITH_OGGVORBIS_DLL (PLATFORM_WINDOWS || PLATFORM_HOLOLENS)
-#endif
-
-
-namespace VorbisChannelInfo
-{
-	extern ENGINE_API const int32 Order[8][8];
-}
-
 /**
- * Loads vorbis dlls
-*/
-ENGINE_API void LoadVorbisLibraries();
-
-#if WITH_OGGVORBIS
-
-/** 
- * Helper class to parse ogg vorbis data
+ * DEPRECATED 5.4 VorbisAudioInfo.h. please use the module path of Decoders/VorbisAudioInfo.h instead.
  */
-class FVorbisAudioInfo : public ICompressedAudioInfo
-{
-public:
-	ENGINE_API FVorbisAudioInfo( void );
-	ENGINE_API virtual ~FVorbisAudioInfo( void );
-
-	/** Emulate read from memory functionality */
-	size_t			ReadMemory( void *ptr, uint32 size );
-	int				SeekMemory( uint32 offset, int whence );
-	int				CloseMemory( void );
-	long			TellMemory( void );
-
-	/** Emulate read from streaming functionality */
-	size_t			ReadStreaming( void *ptr, uint32 size );
-	int				CloseStreaming( void );
-
-	/** Common info and data functions between ReadCompressedInfo/ReadCompressedData and StreamCompressedInfo/StreamCompressedData */
-	/** Can't use ov_callbacks here so we have to pass in a void* for the callbacks instead */
-	bool GetCompressedInfoCommon(void*	Callbacks, FSoundQualityInfo* QualityInfo);
-
-	/** 
-	 * Reads the header information of an ogg vorbis file
-	 * 
-	 * @param	Resource		Info about vorbis data
-	 */
-	ENGINE_API virtual bool ReadCompressedInfo( const uint8* InSrcBufferData, uint32 InSrcBufferDataSize, struct FSoundQualityInfo* QualityInfo ) override;
-
-	/** 
-	 * Decompresses ogg data to raw PCM data. 
-	 * 
-	 * @param	Destination	where to place the decompressed sound
-	 * @param	bLooping	whether to loop the sound by seeking to the start, or pad the buffer with zeroes
-	 * @param	BufferSize	number of bytes of PCM data to create
-	 *
-	 * @return	bool		true if the end of the data was reached (for both single shot and looping sounds)
-	 */
-	ENGINE_API virtual bool ReadCompressedData( uint8* InDestination, bool bLooping, uint32 BufferSize ) override;
-
-	ENGINE_API virtual void SeekToTime( const float SeekTime ) override;
-
-	ENGINE_API virtual void SeekToFrame(const uint32 SeekFrame) override;
-
-	/** 
-	 * Decompress an entire ogg data file to a TArray
-	 */
-	ENGINE_API virtual void ExpandFile( uint8* DstBuffer, struct FSoundQualityInfo* QualityInfo ) override;
-
-	/** 
-	 * Sets ogg to decode to half-rate
-	 * 
-	 * @param	Resource		Info about vorbis data
-	 */
-	ENGINE_API virtual void EnableHalfRate( bool HalfRate ) override;
-
-	virtual uint32 GetSourceBufferSize() const override { return SrcBufferDataSize; }
-
-	virtual bool UsesVorbisChannelOrdering() const override { return true; }
-	virtual int GetStreamBufferSize() const override { return MONO_PCM_BUFFER_SIZE; }
-
-	// Additional overrides for streaming
-	virtual bool SupportsStreaming() const override {return true;}
-	virtual bool StreamCompressedData(uint8* InDestination, bool bLooping, uint32 BufferSize, int32& OutNumBytesStreamed) override;
-	virtual int32 GetCurrentChunkIndex() const override {return CurrentStreamingChunkIndex;}
-	virtual int32 GetCurrentChunkOffset() const override {return BufferOffset % CurrentStreamingChunksSize;}
-	virtual bool HasError() const override;
-	// End of ICompressedAudioInfo Interface
-
-protected:
-	virtual bool StreamCompressedInfoInternal(const FSoundWaveProxyPtr& InWaveProxy, struct FSoundQualityInfo* QualityInfo) override;
-
-	friend class FAudioFormatOgg;
-	ENGINE_API int32 GetAudioDataStartOffset() const;
-
-private:
-	using Super = ICompressedAudioInfo; 
-	const uint8* GetLoadedChunk(FSoundWaveProxyPtr InSoundWave, uint32 ChunkIndex, uint32& OutChunkSize);
-
-	struct FVorbisFileWrapper* VFWrapper;
-	const uint8* SrcBufferData;
-	uint32 SrcBufferDataSize;
-	uint32 BufferOffset;
-	uint32 CurrentBufferChunkOffset;
-
-	// In case Ogg Vorbis fails to return any compressed audio for an asset,
-	// we use this counter to exit out of the decoder's while loop early.
-	int32 TimesLoopedWithoutDecompressedAudio;
-
-	/** Critical section used to prevent multiple threads accessing same ogg-vorbis file handles at the same time */
-	mutable FCriticalSection VorbisCriticalSection;
-
-	uint8 const* CurrentStreamingChunkData;
-	int32 CurrentStreamingChunkIndex;
-	int32 NextStreamingChunkIndex;
-	uint32 CurrentStreamingChunksSize;
-	bool bHeaderParsed;
-	bool bHasError = false;
-	int32 LastErrorCode = 0;
+#pragma once
 
-	// This handle is used to ensure that a chunk currently being decoded isn't evicted until we are done with it.
-	FAudioChunkHandle CurCompressedChunkHandle;
-};
+#ifdef _MSC_VER
+	#pragma message(__FILE__"(9): has moved to it's own module VorbisAudioDecoder, please update your code.")
+#else
+	#pragma message("VorbisAudioInfo.h has moved to it's own module VorbisAudioDecoder, please update your code.")
 #endif
+
+ // Include the new file so that the project still compiles, but has warnings
+#include "Decoders/VorbisAudioInfo.h"
diff --git a/Engine/Source/Runtime/IOS/IOSAudio/Private/IOSAudioBuffer.cpp b/Engine/Source/Runtime/IOS/IOSAudio/Private/IOSAudioBuffer.cpp
--- a/Engine/Source/Runtime/IOS/IOSAudio/Private/IOSAudioBuffer.cpp
+++ b/Engine/Source/Runtime/IOS/IOSAudio/Private/IOSAudioBuffer.cpp
@@ -15,7 +15,6 @@
 #include "AudioDeviceManager.h"
 #include "Engine/Engine.h"
 #include "AudioDecompress.h"
-#include "ADPCMAudioInfo.h"
 
 /*------------------------------------------------------------------------------------
 	FIOSAudioSoundBuffer
@@ -36,7 +35,7 @@ FIOSAudioSoundBuffer::FIOSAudioSoundBuffer(FIOSAudioDevice* InAudioDevice, USoun
 			return;
 		}
 
-		SoundFormat = static_cast<ESoundFormat>(DecompressionState->GetFormatTag());
+		SoundFormat = SoundFormat_LPCM; 
 	}
     else
     {
diff --git a/Engine/Source/Runtime/IOS/IOSAudio/Private/IOSAudioDevice.cpp b/Engine/Source/Runtime/IOS/IOSAudio/Private/IOSAudioDevice.cpp
--- a/Engine/Source/Runtime/IOS/IOSAudio/Private/IOSAudioDevice.cpp
+++ b/Engine/Source/Runtime/IOS/IOSAudio/Private/IOSAudioDevice.cpp
@@ -10,8 +10,8 @@
 
 #include "IOSAudioDevice.h"
 #include "AudioEffect.h"
-#include "ADPCMAudioInfo.h"
 #include "AudioPluginUtilities.h"
+#include "AudioDecompress.h"
 
 DEFINE_LOG_CATEGORY(LogIOSAudio);
 
diff --git a/Engine/Source/Runtime/IOS/IOSAudio/Private/IOSAudioSource.cpp b/Engine/Source/Runtime/IOS/IOSAudio/Private/IOSAudioSource.cpp
--- a/Engine/Source/Runtime/IOS/IOSAudio/Private/IOSAudioSource.cpp
+++ b/Engine/Source/Runtime/IOS/IOSAudio/Private/IOSAudioSource.cpp
@@ -13,7 +13,6 @@
 #include "ContentStreaming.h"
 #include "Interfaces/IAudioFormat.h"
 #include "AudioDecompress.h"
-#include "ADPCMAudioInfo.h"
 #include "Audio.h"
 
 const uint32	Callback_Free = 0;
@@ -483,7 +482,7 @@ OSStatus FIOSAudioSoundSource::IOSAudioRenderCallback(void* RefCon, AudioUnitRen
 		}
 		else
 		{
-			Source->IOSBuffer->RenderCallbackBufferSize = NumFrames * sizeof(uint16) * Source->IOSBuffer->DecompressionState->GetNumChannels();
+			Source->IOSBuffer->RenderCallbackBufferSize = NumFrames * sizeof(uint16) * Source->IOSBuffer->NumChannels;
 
 			// Since StreamCompressedData returns interlaced samples we need to decompress all frames(samples) for all channels here so we don't end up decompressing multiple times
 			// Ensure we have enough memory to do this. If needed we could realloc here but that is bad practice inside the audio callback since it has a hard deadline
diff --git a/Engine/Source/Runtime/IOS/IOSAudio/Public/IOSAudioDevice.h b/Engine/Source/Runtime/IOS/IOSAudio/Public/IOSAudioDevice.h
--- a/Engine/Source/Runtime/IOS/IOSAudio/Public/IOSAudioDevice.h
+++ b/Engine/Source/Runtime/IOS/IOSAudio/Public/IOSAudioDevice.h
@@ -34,6 +34,7 @@ DECLARE_LOG_CATEGORY_EXTERN(LogIOSAudio, Log, All);
 struct FWaveInstance;
 class FIOSAudioDevice;
 class FIOSAudioEffectsManager;
+class ICompressedAudioInfo;
 
 enum ESoundFormat
 {
@@ -94,7 +95,7 @@ public:
 	uint32 BufferSize;
 	
 	/** Wrapper to assist in the bookkeeping of uncompressed data when streaming */
-	class FADPCMAudioInfo*		DecompressionState;
+	ICompressedAudioInfo*		DecompressionState;
 	bool	bStreaming;
 	bool    bIsProcedural;
 };
diff --git a/Engine/Source/Runtime/OpusAudioDecoder/Module/OpusAudioDecoder.Build.cs b/Engine/Source/Runtime/OpusAudioDecoder/Module/OpusAudioDecoder.Build.cs
new file mode 100644
--- /dev/null
+++ b/Engine/Source/Runtime/OpusAudioDecoder/Module/OpusAudioDecoder.Build.cs
@@ -0,0 +1,27 @@
+// Copyright Epic Games, Inc. All Rights Reserved.
+
+using UnrealBuildTool;
+
+public class OpusAudioDecoder : ModuleRules
+{
+    public OpusAudioDecoder(ReadOnlyTargetRules Target) : base(Target)
+    {
+        ShortName = "OpusAudioDecoder";
+
+        PublicDependencyModuleNames.AddRange(
+            new string[]
+            {
+                "Core",
+                "CoreUObject",
+                "Engine"
+            }
+        );
+        
+        AddEngineThirdPartyPrivateStaticDependencies(Target,
+	        "UEOgg",
+	        "libOpus"
+        );        
+
+	    PublicDefinitions.Add("WITH_OPUS=1");
+    }
+}
diff --git a/Engine/Source/Runtime/OpusAudioDecoder/Module/Private/OpusAudioInfo.cpp b/Engine/Source/Runtime/OpusAudioDecoder/Module/Private/OpusAudioInfo.cpp
new file mode 100644
--- /dev/null
+++ b/Engine/Source/Runtime/OpusAudioDecoder/Module/Private/OpusAudioInfo.cpp
@@ -0,0 +1,396 @@
+// Copyright Epic Games, Inc. All Rights Reserved.
+
+
+#include "Decoders/OpusAudioInfo.h"
+#include "Interfaces/IAudioFormat.h"
+#include <opus_defines.h>
+#include <opus_types.h>
+
+THIRD_PARTY_INCLUDES_START
+#include "opus_multistream.h"
+THIRD_PARTY_INCLUDES_END
+
+#define USE_UE_MEM_ALLOC 1
+
+DEFINE_LOG_CATEGORY_STATIC(LogOpusAudioDecoder, Log, All);
+
+///////////////////////////////////////////////////////////////////////////////////////
+// Followed pattern used in opus_multistream_encoder.c - this will allow us to setup //
+// a multistream decoder without having to save extra information for every asset.   //
+///////////////////////////////////////////////////////////////////////////////////////
+struct UnrealChannelLayout{
+	int32 NumStreams;
+	int32 NumCoupledStreams;
+	uint8 Mapping[8];
+};
+
+/* Index is NumChannels-1*/
+static const UnrealChannelLayout UnrealMappings[8] = {
+	{ 1, 0, { 0 } },                      /* 1: mono */
+	{ 1, 1, { 0, 1 } },                   /* 2: stereo */
+	{ 2, 1, { 0, 1, 2 } },                /* 3: 1-d surround */
+	{ 2, 2, { 0, 1, 2, 3 } },             /* 4: quadraphonic surround */
+	{ 3, 2, { 0, 1, 4, 2, 3 } },          /* 5: 5-channel surround */
+	{ 4, 2, { 0, 1, 4, 5, 2, 3 } },       /* 6: 5.1 surround */
+	{ 4, 3, { 0, 1, 4, 6, 2, 3, 5 } },    /* 7: 6.1 surround */
+	{ 5, 3, { 0, 1, 6, 7, 2, 3, 4, 5 } }, /* 8: 7.1 surround */
+};
+
+/*------------------------------------------------------------------------------------
+FOpusDecoderWrapper
+------------------------------------------------------------------------------------*/
+struct FOpusDecoderWrapper
+{
+	FOpusDecoderWrapper(uint32 SampleRate, uint8 NumChannels)
+	{
+		check(NumChannels <= 8);
+		const UnrealChannelLayout& Layout = UnrealMappings[NumChannels-1];
+	#if USE_UE_MEM_ALLOC
+		int32 DecSize = opus_multistream_decoder_get_size(Layout.NumStreams, Layout.NumCoupledStreams);
+		Decoder = (OpusMSDecoder*)FMemory::Malloc(DecSize);
+		DecError = opus_multistream_decoder_init(Decoder, SampleRate, NumChannels, Layout.NumStreams, Layout.NumCoupledStreams, Layout.Mapping);
+	#else
+		Decoder = opus_multistream_decoder_create(SampleRate, NumChannels, Layout.NumStreams, Layout.NumCoupledStreams, Layout.Mapping, &DecError);
+	#endif
+	}
+
+	~FOpusDecoderWrapper()
+	{
+	#if USE_UE_MEM_ALLOC
+		FMemory::Free(Decoder);
+	#else
+		opus_multistream_decoder_destroy(Decoder);
+	#endif
+	}
+
+	int32 Decode(const uint8* FrameData, uint16 FrameSize, int16* OutPCMData, int32 SampleSize)
+	{
+		return opus_multistream_decode(Decoder, FrameData, FrameSize, OutPCMData, SampleSize, 0);
+	}
+
+	bool WasInitialisedSuccessfully() const
+	{
+		return DecError == OPUS_OK;
+	}
+
+private:
+	OpusMSDecoder* Decoder;
+	int32 DecError;
+};
+
+/*------------------------------------------------------------------------------------
+FOpusAudioInfo.
+------------------------------------------------------------------------------------*/
+FOpusAudioInfo::FOpusAudioInfo()
+	: OpusDecoderWrapper(nullptr)
+{
+}
+
+FOpusAudioInfo::~FOpusAudioInfo()
+{
+	if (OpusDecoderWrapper != nullptr)
+	{
+		delete OpusDecoderWrapper;
+		OpusDecoderWrapper = nullptr;
+	}
+}
+
+
+bool FOpusAudioInfo::ParseHeader(FHeader& OutHeader, uint32& OutNumRead, const uint8* InSrcBufferData, uint32 InSrcBufferDataSize)
+{
+	OutNumRead = 0;
+
+	if ((int32)InSrcBufferDataSize < FHeader::HeaderSize())
+	{
+		return false;
+	}
+
+	auto Read = [&InSrcBufferData, &OutNumRead](void* To, int32 NumBytes) -> void
+	{
+		FMemory::Memcpy(To, InSrcBufferData, NumBytes);
+		InSrcBufferData += NumBytes;
+		OutNumRead += NumBytes;
+	};
+
+	Read(OutHeader.Identifier, 8);
+	if (FMemory::Memcmp(OutHeader.Identifier, FHeader::OPUS_ID, 8))
+	{
+		return false;
+	}
+	Read(&OutHeader.Version, sizeof(uint8));
+	Read(&OutHeader.NumChannels, sizeof(uint8));
+	Read(&OutHeader.SampleRate, sizeof(uint32));
+	Read(&OutHeader.EncodedSampleRate, sizeof(uint32));
+	Read(&OutHeader.ActiveSampleCount, sizeof(uint64));
+	Read(&OutHeader.NumEncodedFrames, sizeof(uint32));
+	Read(&OutHeader.NumPreSkipSamples, sizeof(int32));
+	Read(&OutHeader.NumSilentSamplesAtBeginning, sizeof(int32));
+	Read(&OutHeader.NumSilentSamplesAtEnd, sizeof(int32));
+	return true;
+}
+
+bool FOpusAudioInfo::ParseHeader(const uint8* InSrcBufferData, uint32 InSrcBufferDataSize, struct FSoundQualityInfo* QualityInfo)
+{
+	SrcBufferData = InSrcBufferData;
+	SrcBufferDataSize = InSrcBufferDataSize;
+	SrcBufferOffset = 0;
+	CurrentSampleCount = 0;
+
+	Header.Reset();
+	if (!ParseHeader(Header, SrcBufferOffset, InSrcBufferData, InSrcBufferDataSize))
+	{
+		return false;
+	}
+
+	// Store the offset to where the audio data begins
+	AudioDataOffset = SrcBufferOffset;
+	// Sample counts in the Opus API are always per-channel so we multiply by NumChannels here
+	TrueSampleCount = (uint32)Header.ActiveSampleCount * Header.NumChannels;
+	NumChannels = Header.NumChannels;
+
+	// Write out the the header info
+	if (QualityInfo)
+	{
+		QualityInfo->SampleRate = Header.SampleRate;
+		QualityInfo->NumChannels = Header.NumChannels;
+		QualityInfo->SampleDataSize = (uint32)Header.ActiveSampleCount * QualityInfo->NumChannels * sizeof(int16);
+		QualityInfo->Duration = (float)((double)Header.ActiveSampleCount / QualityInfo->SampleRate);
+	}
+
+	return true;
+}
+
+bool FOpusAudioInfo::CreateDecoder()
+{
+	check(OpusDecoderWrapper == nullptr);
+	OpusDecoderWrapper = new FOpusDecoderWrapper(Header.EncodedSampleRate, NumChannels);
+	if (!OpusDecoderWrapper->WasInitialisedSuccessfully())
+	{
+		delete OpusDecoderWrapper;
+		OpusDecoderWrapper = nullptr;
+		return false;
+	}
+
+	NumRemainingSamplesToSkip = Header.NumSilentSamplesAtBeginning + Header.NumPreSkipSamples;
+	PreviousDecodedUnusedSamples.Empty();
+	return true;
+}
+
+int32 FOpusAudioInfo::GetFrameSize()
+{
+	// Opus format has variable frame size at the head of each frame...
+	// We have to assume that the SrcBufferOffset is at the correct location for the read
+	uint16 FrameSize = 0;
+	Read(&FrameSize, sizeof(uint16));
+	return (int32)FrameSize;
+}
+
+uint32 FOpusAudioInfo::GetMaxFrameSizeSamples() const
+{
+	// The encoder is using 20ms frame sizes.
+	const int32 OPUS_MAX_FRAME_SIZE_MS = 20;
+	// There can be at most 2 frames in one packet, so multiply by 2.
+	return Header.EncodedSampleRate * OPUS_MAX_FRAME_SIZE_MS * 2 / 1000;
+}
+
+FDecodeResult FOpusAudioInfo::Decode(const uint8* CompressedData, const int32 CompressedDataSize, uint8* OutPCMData, const int32 OutputPCMDataSize)
+{
+	FDecodeResult Result;
+
+	if (OpusDecoderWrapper)
+	{
+		const int32 kFrameBytes = NumChannels * sizeof(int16);
+		int32 OutputSizeToGo = OutputPCMDataSize;
+		int32 CompressedSizeToGo = CompressedDataSize;
+		const uint8* InputDataPtr = CompressedData; 
+		uint8* OutputDataPtr = OutPCMData;
+
+		Result.NumAudioFramesProduced = 0;
+		while(OutputSizeToGo > 0)
+		{
+			// Get anything that is left over from the previous call.
+			if (PreviousDecodedUnusedSamples.Num())
+			{
+				check(NumRemainingSamplesToSkip == 0);
+				int32 MaxToCopyOut = OutputSizeToGo >= PreviousDecodedUnusedSamples.Num() ? PreviousDecodedUnusedSamples.Num() : OutputSizeToGo;
+				FMemory::Memcpy(OutputDataPtr, PreviousDecodedUnusedSamples.GetData(), MaxToCopyOut);
+				PreviousDecodedUnusedSamples.RemoveAt(0, MaxToCopyOut);
+				OutputSizeToGo -= MaxToCopyOut;
+				OutputDataPtr += MaxToCopyOut;
+				// If there is still something left over then we are done here.
+				Result.NumAudioFramesProduced += MaxToCopyOut / kFrameBytes;
+				if (OutputSizeToGo == 0 || PreviousDecodedUnusedSamples.Num())
+				{
+					Result.NumCompressedBytesConsumed = InputDataPtr - CompressedData;
+					Result.NumPcmBytesProduced = Result.NumAudioFramesProduced * kFrameBytes;
+					return Result;
+				}
+			}
+			// Something left to decompress?
+			if (CompressedSizeToGo <= 0)
+			{
+				break;
+			}
+
+			// Decode the next chunk.
+			const int32 AvailableSampleOutputSize = OutputSizeToGo / kFrameBytes;
+
+			int32 ChunkToDecodeSize = 0;
+			int32 ActualChunkSize = 0;
+			// Is this streaming?
+			const uint8* ChunkToDecodePtr = InputDataPtr;
+			if (!bIsStreaming)
+			{
+				// When not streaming each input data chunk is prepended with the size of the chunk.
+				ChunkToDecodeSize = (int32)InputDataPtr[0] + ((int32)InputDataPtr[1] << 8);
+				ActualChunkSize = ChunkToDecodeSize + 2;
+				ChunkToDecodePtr += 2;
+			}
+			else
+			{
+				// When streaming the chunk size is still prepended to the chunk, but is skipped over
+				// by the caller and provided as an argument.
+				ChunkToDecodeSize = CompressedDataSize;
+				ActualChunkSize = ChunkToDecodeSize;
+			}
+
+
+			int32 NumDecodedFrames = OpusDecoderWrapper->Decode(ChunkToDecodePtr, ChunkToDecodeSize, (int16*)OutputDataPtr, AvailableSampleOutputSize);
+			if (NumDecodedFrames >= 0)
+			{
+				CompressedSizeToGo -= ActualChunkSize;
+				InputDataPtr += ActualChunkSize;
+				if (NumRemainingSamplesToSkip)
+				{
+					if (NumRemainingSamplesToSkip >= NumDecodedFrames)
+					{
+						NumRemainingSamplesToSkip -= NumDecodedFrames;
+						NumDecodedFrames = 0;
+					}
+					else
+					{
+						uint8* FirstUsable = OutputDataPtr + NumRemainingSamplesToSkip * kFrameBytes;
+						int32 UsableSize = (NumDecodedFrames - NumRemainingSamplesToSkip) * kFrameBytes;
+						FMemory::Memmove(OutputDataPtr, FirstUsable, UsableSize);
+						NumDecodedFrames -= NumRemainingSamplesToSkip;
+						NumRemainingSamplesToSkip = 0;
+					}
+				}
+				Result.NumAudioFramesProduced += NumDecodedFrames;
+				OutputSizeToGo -= NumDecodedFrames * kFrameBytes;
+				OutputDataPtr += NumDecodedFrames * kFrameBytes;
+			}
+			// Was the remaining output buffer too small?
+			else if (NumDecodedFrames == OPUS_BUFFER_TOO_SMALL)
+			{
+				int32 NumSampleSpaceNeeded = GetMaxFrameSizeSamples();
+				int32 NumPrevBefore = PreviousDecodedUnusedSamples.Num();
+				PreviousDecodedUnusedSamples.AddUninitialized(NumSampleSpaceNeeded * kFrameBytes);
+				uint8* TempPtr = PreviousDecodedUnusedSamples.GetData() + NumPrevBefore;
+				NumDecodedFrames = OpusDecoderWrapper->Decode(ChunkToDecodePtr, ChunkToDecodeSize, (int16*)TempPtr, NumSampleSpaceNeeded);
+				if (NumDecodedFrames >= 0)
+				{
+					int32 NumPrevNow = NumDecodedFrames * kFrameBytes;
+					PreviousDecodedUnusedSamples.SetNum(NumPrevNow);
+					CompressedSizeToGo -= ActualChunkSize;
+					InputDataPtr += ActualChunkSize;
+					continue;
+				}
+				else
+				{
+					UE_LOG(LogOpusAudioDecoder, Error, TEXT("opus_multistream_decode() returned %d while decoding into temporary buffer"), NumDecodedFrames);
+					return FDecodeResult();
+				}
+			}
+			else
+			{
+				// A decode error of sorts.
+				UE_LOG(LogOpusAudioDecoder, Error, TEXT("opus_multistream_decode() returned %d"), NumDecodedFrames);
+				return FDecodeResult();
+			}
+		}
+		Result.NumCompressedBytesConsumed = InputDataPtr - CompressedData;
+		Result.NumPcmBytesProduced = Result.NumAudioFramesProduced * kFrameBytes;
+	}
+	return Result;
+}
+
+void FOpusAudioInfo::PrepareToLoop()
+{
+	IStreamedCompressedInfo::PrepareToLoop();
+	NumRemainingSamplesToSkip = Header.NumSilentSamplesAtBeginning + Header.NumPreSkipSamples;
+	PreviousDecodedUnusedSamples.Empty();
+}
+
+void FOpusAudioInfo::SeekToTime(const float InSeekTime)
+{
+	if (GetStreamingSoundWave().IsValid())
+	{
+		IStreamedCompressedInfo::SeekToTime(InSeekTime);
+	}
+	else if (SrcBufferData && SrcBufferDataSize)
+	{
+		uint32 SeekFrameNum = 0;
+		if (InSeekTime > 0.0f)
+		{
+			SeekFrameNum = (uint32)(InSeekTime * Header.SampleRate);
+		}
+
+		SeekToFrame(SeekFrameNum);
+	}
+}
+
+void FOpusAudioInfo::SeekToFrame(const uint32 InSeekFrame)
+{
+	if (GetStreamingSoundWave().IsValid())
+	{
+		IStreamedCompressedInfo::SeekToFrame(InSeekFrame);
+	}
+	else if (SrcBufferData && SrcBufferDataSize)
+	{
+		uint32 SeekSampleNum = InSeekFrame * NumChannels;
+		const uint8* ChunkPtr = SrcBufferData + AudioDataOffset;
+		const uint8* const EndPtr = SrcBufferData + SrcBufferDataSize;
+		uint32 CurrentChunkSampleNum = 0;
+		while (ChunkPtr < EndPtr)
+		{
+			uint32 ChunkSize = (uint32)ChunkPtr[0] + ((uint32)ChunkPtr[1] << 8);
+			int32 ExpectedFrames = opus_packet_get_nb_frames(ChunkPtr + 2, ChunkSize);
+			int32 ExpectedFrameSize = opus_packet_get_samples_per_frame(ChunkPtr + 2, Header.EncodedSampleRate);
+			int32 NumExpectedTotal = ExpectedFrames * ExpectedFrameSize * NumChannels;
+
+			if (CurrentChunkSampleNum >= SeekSampleNum && SeekSampleNum < CurrentChunkSampleNum + NumExpectedTotal)
+			{
+				CurrentSampleCount = CurrentChunkSampleNum;
+				SrcBufferOffset = ChunkPtr - SrcBufferData;
+				break;
+			}
+			CurrentChunkSampleNum += NumExpectedTotal;
+			ChunkPtr += ChunkSize + 2;
+		}
+		// Not found?
+		if (ChunkPtr >= EndPtr)
+		{
+			CurrentSampleCount = SeekSampleNum;
+			SrcBufferOffset = SrcBufferDataSize;
+		}
+	}
+
+	NumRemainingSamplesToSkip = Header.NumSilentSamplesAtBeginning + Header.NumPreSkipSamples;
+	PreviousDecodedUnusedSamples.Empty();
+}
+
+class OPUSAUDIODECODER_API FOpusAudioDecoderModule : public IModuleInterface
+{
+public:	
+	TUniquePtr<IAudioInfoFactory> Factory;
+
+	virtual void StartupModule() override
+	{
+		Factory = MakeUnique<FSimpleAudioInfoFactory>([] { return new FOpusAudioInfo(); }, Audio::NAME_OPUS);
+	}
+
+	virtual void ShutdownModule() override {}
+};
+
+IMPLEMENT_MODULE(FOpusAudioDecoderModule, OpusAudioDecoder)
diff --git a/Engine/Source/Runtime/OpusAudioDecoder/Module/Public/Decoders/OpusAudioInfo.h b/Engine/Source/Runtime/OpusAudioDecoder/Module/Public/Decoders/OpusAudioInfo.h
--- a/Engine/Source/Runtime/OpusAudioDecoder/Module/Public/Decoders/OpusAudioInfo.h
+++ b/Engine/Source/Runtime/OpusAudioDecoder/Module/Public/Decoders/OpusAudioInfo.h
@@ -18,9 +18,8 @@ struct FSoundQualityInfo;
 class FOpusAudioInfo : public IStreamedCompressedInfo
 {
 public:
-	ENGINE_API FOpusAudioInfo();
-	ENGINE_API virtual ~FOpusAudioInfo();
-
+	OPUSAUDIODECODER_API FOpusAudioInfo();
+	OPUSAUDIODECODER_API virtual ~FOpusAudioInfo();
 	//~ Begin IStreamedCompressedInfo Interface
 	bool ParseHeader(const uint8* InSrcBufferData, uint32 InSrcBufferDataSize, FSoundQualityInfo* QualityInfo) override;
 	int32 GetFrameSize() override;
@@ -31,8 +30,7 @@ public:
 	void SeekToTime(const float SeekTime) override;
 	void SeekToFrame(const uint32 SeekFrame) override;
 	//~ End IStreamedCompressedInfo Interface
-
-	struct ENGINE_API FHeader
+	struct OPUSAUDIODECODER_API FHeader
 	{
 		char Identifier[8];
 		uint8 Version = 0;
@@ -75,8 +73,7 @@ public:
 				+ sizeof(int32);	// NumSilentSamplesAtEnd
 		}
 	};
-	static ENGINE_API bool ParseHeader(FHeader& OutHeader, uint32& OutNumRead, const uint8* InSrcBufferData, uint32 InSrcBufferDataSize);
-
+	static OPUSAUDIODECODER_API bool ParseHeader(FHeader& OutHeader, uint32& OutNumRead, const uint8* InSrcBufferData, uint32 InSrcBufferDataSize);
 protected:
 	/** Wrapper around Opus-specific decoding state and APIs */
 	FOpusDecoderWrapper* OpusDecoderWrapper;
@@ -89,4 +86,4 @@ protected:
 
 	/** Samples that are left over from the previous Decode() call that did not fit into the receive buffer. */
 	TArray<uint8> PreviousDecodedUnusedSamples;
-};
+};
\ No newline at end of file
diff --git a/Engine/Source/Runtime/VorbisAudioDecoder/Module/Private/VorbisAudioInfo.cpp b/Engine/Source/Runtime/VorbisAudioDecoder/Module/Private/VorbisAudioInfo.cpp
new file mode 100644
--- /dev/null
+++ b/Engine/Source/Runtime/VorbisAudioDecoder/Module/Private/VorbisAudioInfo.cpp
@@ -0,0 +1,861 @@
+// Copyright Epic Games, Inc. All Rights Reserved.
+
+
+#include "Decoders/VorbisAudioInfo.h"
+#include "Interfaces/IAudioFormat.h"
+#include "Audio.h"
+#include <ogg/os_types.h>
+#include <vorbis/codec.h>
+#include "Misc/Paths.h"
+#if PLATFORM_WINDOWS
+#include "Windows/WindowsHWrapper.h"
+#endif
+
+#if WITH_OGGVORBIS
+
+#pragma pack(push, 8)
+#include "vorbis/vorbisfile.h"
+#pragma pack(pop)
+#endif
+
+#if PLATFORM_LITTLE_ENDIAN
+#define VORBIS_BYTE_ORDER 0
+#else
+#define VORBIS_BYTE_ORDER 1
+#endif
+
+// Non-windows platform don't load Dlls
+#if !WITH_OGGVORBIS_DLL
+static FThreadSafeBool bDllLoaded = true;
+#else
+static FThreadSafeBool bDllLoaded;
+#endif
+
+static int32 VorbisReadFailiureTimeoutCVar = 1;
+FAutoConsoleVariableRef CVarVorbisReadFailiureTimeout(
+	TEXT("au.vorbis.ReadFailiureTimeout"),
+	VorbisReadFailiureTimeoutCVar,
+	TEXT("When set to 1, we bail on decoding Ogg Vorbis sounds if we were not able to successfully decode them after several attempts.\n"),
+	ECVF_Default);
+
+/**
+ * Channel order expected for a multi-channel ogg vorbis file.
+ * Ordering taken from http://xiph.org/vorbis/doc/Vorbis_I_spec.html#x1-800004.3.9
+ */
+namespace VorbisChannelInfo
+{
+	const int32 Order[8][8] = {
+		{ 0 },
+		{ 0, 1 },
+		{ 0, 2, 1 },
+		{ 0, 1, 2, 3 },
+		{ 0, 2, 1, 3, 4 },
+		{ 0, 2, 1, 4, 5, 3 },
+		{ 0, 2, 1, 4, 5, 6, 3 },
+		{ 0, 2, 1, 4, 5, 6, 7, 3 }
+	};
+};
+
+/*------------------------------------------------------------------------------------
+	FVorbisFileWrapper. Hides Vorbis structs from public headers.
+------------------------------------------------------------------------------------*/
+struct FVorbisFileWrapper
+{
+	FVorbisFileWrapper()
+	{
+#if WITH_OGGVORBIS
+		FMemory::Memzero( &vf, sizeof( OggVorbis_File ) );
+#endif
+	}
+
+	~FVorbisFileWrapper()
+	{
+#if WITH_OGGVORBIS
+		// Only clear vorbis if the DLL succeeded in loading
+		if (bDllLoaded)
+		{
+			ov_clear(&vf);
+		}
+#endif
+	}
+
+#if WITH_OGGVORBIS
+	/** Ogg vorbis decompression state */
+	OggVorbis_File	vf;
+#endif
+};
+
+#if WITH_OGGVORBIS
+/*------------------------------------------------------------------------------------
+	FVorbisAudioInfo.
+------------------------------------------------------------------------------------*/
+FVorbisAudioInfo::FVorbisAudioInfo()
+	: VFWrapper(new FVorbisFileWrapper())
+	, SrcBufferData(NULL)
+	, SrcBufferDataSize(0)
+	, BufferOffset(0)
+	, CurrentBufferChunkOffset(0)
+	, TimesLoopedWithoutDecompressedAudio(0)
+	, CurrentStreamingChunkData(nullptr)
+	, CurrentStreamingChunkIndex(INDEX_NONE)
+	, NextStreamingChunkIndex(0)
+	, CurrentStreamingChunksSize(0)
+	, bHeaderParsed(false)
+{
+	// Make sure we have properly allocated a VFWrapper
+	check(VFWrapper != NULL);
+}
+
+FVorbisAudioInfo::~FVorbisAudioInfo()
+{
+	FScopeLock ScopeLock(&VorbisCriticalSection);
+	check(VFWrapper != nullptr);
+	delete VFWrapper;
+	VFWrapper = nullptr;
+}
+
+/** Emulate read from memory functionality */
+size_t FVorbisAudioInfo::ReadMemory( void *Ptr, uint32 Size )
+{
+	check(Ptr);
+	size_t BytesToRead = FMath::Min(Size, SrcBufferDataSize - BufferOffset);
+	FMemory::Memcpy( Ptr, SrcBufferData + BufferOffset, BytesToRead );
+	BufferOffset += BytesToRead;
+	return( BytesToRead );
+}
+
+static size_t OggReadMemory( void *ptr, size_t size, size_t nmemb, void *datasource )
+{
+	check(ptr);
+	check(datasource);
+	FVorbisAudioInfo* OggInfo = (FVorbisAudioInfo*)datasource;
+	return( OggInfo->ReadMemory( ptr, size * nmemb ) );
+}
+
+int FVorbisAudioInfo::SeekMemory( uint32 offset, int whence )
+{
+	switch( whence )
+	{
+	case SEEK_SET:
+		BufferOffset = offset;
+		break;
+
+	case SEEK_CUR:
+		BufferOffset += offset;
+		break;
+
+	case SEEK_END:
+		BufferOffset = SrcBufferDataSize - offset;
+		break;
+
+	default:
+		checkf(false, TEXT("Uknown seek type"));
+		break;
+	}
+
+	BufferOffset = FMath::Clamp<uint32>(BufferOffset, 0, SrcBufferDataSize);
+	return( BufferOffset );
+}
+
+static int OggSeekMemory( void *datasource, ogg_int64_t offset, int whence )
+{
+	FVorbisAudioInfo* OggInfo = ( FVorbisAudioInfo* )datasource;
+	return( OggInfo->SeekMemory( offset, whence ) );
+}
+
+int FVorbisAudioInfo::CloseMemory( void )
+{
+	return( 0 );
+}
+
+static int OggCloseMemory( void *datasource )
+{
+	FVorbisAudioInfo* OggInfo = ( FVorbisAudioInfo* )datasource;
+	return( OggInfo->CloseMemory() );
+}
+
+long FVorbisAudioInfo::TellMemory( void )
+{
+	return( BufferOffset );
+}
+
+static long OggTellMemory( void *datasource )
+{
+	FVorbisAudioInfo *OggInfo = ( FVorbisAudioInfo* )datasource;
+	return( OggInfo->TellMemory() );
+}
+
+/** Emulate read from memory functionality */
+size_t FVorbisAudioInfo::ReadStreaming(void *Ptr, uint32 Size )
+{
+	size_t NumBytesRead = 0;
+
+	while (NumBytesRead < Size)
+	{
+		if (!CurrentStreamingChunkData || CurrentStreamingChunkIndex != NextStreamingChunkIndex)
+		{
+			CurrentStreamingChunkIndex = NextStreamingChunkIndex;
+
+			check(CurrentStreamingChunkIndex >= 0);
+
+			if (static_cast<uint32>(CurrentStreamingChunkIndex) >= StreamingSoundWave->GetNumChunks())
+			{
+				CurrentStreamingChunkData = nullptr;
+				CurrentStreamingChunksSize = 0;
+			}
+			else
+			{
+				CurrentStreamingChunkData = GetLoadedChunk(StreamingSoundWave, CurrentStreamingChunkIndex, CurrentStreamingChunksSize);
+			}
+
+			if (CurrentStreamingChunkData)
+			{
+				check(CurrentStreamingChunkIndex < (int32)StreamingSoundWave->GetNumChunks());
+				CurrentBufferChunkOffset = 0;
+			}
+		}
+
+		// No chunk data -- either looping or something else happened with stream
+		if (!CurrentStreamingChunkData)
+		{
+			return NumBytesRead;
+		}
+
+		check(CurrentStreamingChunksSize >= CurrentBufferChunkOffset);
+		// How many bytes left in the current chunk
+		uint32 BytesLeftInCurrentChunk = CurrentStreamingChunksSize - CurrentBufferChunkOffset;
+
+		check(Size >= NumBytesRead);
+		// How many more bytes we want to read
+		uint32 NumBytesLeftToRead = Size - NumBytesRead;
+
+		// The amount of audio we're going to copy is the min of the bytes left in the chunk and the bytes left we need to read.
+		size_t BytesToCopy = FMath::Min(BytesLeftInCurrentChunk, NumBytesLeftToRead);
+		if (BytesToCopy > 0)
+		{
+			void* WriteBufferLocation = (void*)((uint8*)Ptr + NumBytesRead);
+			FMemory::Memcpy(WriteBufferLocation, CurrentStreamingChunkData + CurrentBufferChunkOffset, BytesToCopy);
+
+			// Increment the BufferOffset by how many bytes we copied from the stream
+			BufferOffset += BytesToCopy;
+
+			// Increment the current buffer's offset
+			CurrentBufferChunkOffset += BytesToCopy;
+
+			// Increment the number of bytes we read this callback.
+			NumBytesRead += BytesToCopy;
+		}
+
+		// If we need to read more bytes than are left in the current chunk, we're going to need to increment the chunk index so we read the next chunk of audio
+		if (NumBytesLeftToRead >= BytesLeftInCurrentChunk)
+		{
+			NextStreamingChunkIndex++;
+		}
+	}
+
+	return NumBytesRead;
+}
+
+static size_t OggReadStreaming( void *ptr, size_t size, size_t nmemb, void *datasource )
+{
+	check(ptr);
+	check(datasource);
+	if (FVorbisAudioInfo* OggInfo = (FVorbisAudioInfo*)datasource)
+	{
+		return OggInfo->ReadStreaming(ptr, size * nmemb);
+	}
+	UE_LOG(LogAudio, Error, TEXT("OggReadStreaming had null audio info datasource."));
+	return 0;
+}
+
+int FVorbisAudioInfo::CloseStreaming( void )
+{
+	return( 0 );
+}
+
+static int OggCloseStreaming( void *datasource )
+{
+	FVorbisAudioInfo* OggInfo = ( FVorbisAudioInfo* )datasource;
+	return( OggInfo->CloseStreaming() );
+}
+
+#define CASE_AND_STRING(X) case X: return TEXT(#X)
+
+static const TCHAR* GetVorbisError(const int32 InErrorCode)
+{
+	switch (InErrorCode)
+	{
+		CASE_AND_STRING(OV_FALSE);
+		CASE_AND_STRING(OV_EOF);
+		CASE_AND_STRING(OV_HOLE);
+
+		CASE_AND_STRING(OV_EREAD);
+		CASE_AND_STRING(OV_EFAULT);
+		CASE_AND_STRING(OV_EIMPL);
+		CASE_AND_STRING(OV_EINVAL);
+		CASE_AND_STRING(OV_ENOTVORBIS);
+		CASE_AND_STRING(OV_EBADHEADER);
+		CASE_AND_STRING(OV_EVERSION);
+		CASE_AND_STRING(OV_ENOTAUDIO);
+		CASE_AND_STRING(OV_EBADPACKET);
+		CASE_AND_STRING(OV_EBADLINK);
+		CASE_AND_STRING(OV_ENOSEEK);
+	default:
+		break;
+	}
+	return TEXT("Unknown");
+}
+
+#undef CASE_AND_STRING
+
+bool FVorbisAudioInfo::GetCompressedInfoCommon(void* Callbacks, FSoundQualityInfo* QualityInfo)
+{
+	if (!bDllLoaded)
+	{
+		UE_LOG(LogAudio, Error, TEXT("FVorbisAudioInfo::GetCompressedInfoCommon failed due to vorbis DLL not being loaded."));
+		bHasError = true;
+		return false;
+	}
+
+	// Set up the read from memory variables
+	int Result = ov_open_callbacks(this, &VFWrapper->vf, NULL, 0, (*(ov_callbacks*)Callbacks));
+	if (Result < 0)
+	{
+		UE_LOG(LogAudio, Warning, TEXT("FVorbisAudioInfo::ReadCompressedInfo, ov_open_callbacks error code: %d : %s"), Result,GetVorbisError(Result));
+		bHasError = true;
+		LastErrorCode = Result;
+		return false;
+	}
+
+	if( QualityInfo )
+	{
+		// The compression could have resampled the source to make loopable
+		vorbis_info* vi = ov_info( &VFWrapper->vf, -1 );
+		QualityInfo->SampleRate = vi->rate;
+		QualityInfo->NumChannels = vi->channels;
+		ogg_int64_t PCMTotal = ov_pcm_total( &VFWrapper->vf, -1 );
+		if (PCMTotal >= 0)
+		{
+			QualityInfo->SampleDataSize = PCMTotal * QualityInfo->NumChannels * sizeof( int16 );
+			QualityInfo->Duration = ( float )ov_time_total( &VFWrapper->vf, -1 );
+		}
+		else if (PCMTotal == OV_EINVAL)
+		{
+			// indicates an error or that the bitstream is non-seekable
+			QualityInfo->SampleDataSize = 0;
+			QualityInfo->Duration = 0.0f;
+		}
+	}
+
+	return true;
+}
+
+/**
+ * Reads the header information of an ogg vorbis file
+ *
+ * @param	Resource		Info about vorbis data
+ */
+bool FVorbisAudioInfo::ReadCompressedInfo( const uint8* InSrcBufferData, uint32 InSrcBufferDataSize, FSoundQualityInfo* QualityInfo )
+{
+	if (!bDllLoaded)
+	{
+		UE_LOG(LogAudio, Error, TEXT("FVorbisAudioInfo::ReadCompressedInfo failed due to vorbis DLL not being loaded."));
+		return false;
+	}
+	
+	if (bHeaderParsed)
+	{
+		UE_LOG(LogAudio, Error, TEXT("FVorbisAudioInfo::ReadCompressedInfo failed due to the header being parsed already."));
+		return false;
+	}
+
+	if (!VFWrapper)
+	{
+		UE_LOG(LogAudio, Error, TEXT("FVorbisAudioInfo::ReadCompressedInfo failed due to not having vorbis wrapper."));
+		return false;
+	}
+	
+
+	FScopeLock ScopeLock(&VorbisCriticalSection);
+
+	ov_callbacks Callbacks;
+
+	SrcBufferData = InSrcBufferData;
+	SrcBufferDataSize = InSrcBufferDataSize;
+	BufferOffset = 0;
+
+	Callbacks.read_func = OggReadMemory;
+	Callbacks.seek_func = OggSeekMemory;
+	Callbacks.close_func = OggCloseMemory;
+	Callbacks.tell_func = OggTellMemory;
+
+	bHeaderParsed = GetCompressedInfoCommon(&Callbacks, QualityInfo);
+
+	if (!bHeaderParsed)
+	{	
+		UE_LOG(LogAudio, Error, TEXT("Failed to parse header for compressed vorbis file."));
+	}
+
+
+	return bHeaderParsed;
+}
+
+
+/**
+ * Decompress an entire ogg vorbis data file to a TArray
+ */
+void FVorbisAudioInfo::ExpandFile( uint8* DstBuffer, FSoundQualityInfo* QualityInfo )
+{
+	if (!bDllLoaded)
+	{
+		UE_LOG(LogAudio, Error, TEXT("FVorbisAudioInfo::ExpandFile failed due to vorbis DLL not being loaded."));
+		return;
+	}
+
+	uint32		TotalBytesRead, BytesToRead;
+
+	check( VFWrapper != NULL );
+	check( DstBuffer );
+	check( QualityInfo );
+
+	FScopeLock ScopeLock(&VorbisCriticalSection);
+
+	if (!bHeaderParsed)
+	{
+		UE_LOG(LogAudio, Error, TEXT("Failed to expand vorbis file to not parsing header first."));
+		return;
+	}
+
+	// A zero buffer size means decompress the entire ogg vorbis stream to PCM.
+	TotalBytesRead = 0;
+	BytesToRead = QualityInfo->SampleDataSize;
+
+	char* Destination = ( char* )DstBuffer;
+	while( TotalBytesRead < BytesToRead )
+	{
+		long BytesRead = ov_read( &VFWrapper->vf, Destination, BytesToRead - TotalBytesRead, 0, 2, 1, NULL );
+
+		if (BytesRead < 0)
+		{
+			// indicates an error - fill remainder of buffer with zero
+			FMemory::Memzero(Destination, BytesToRead - TotalBytesRead);
+			return;
+		}
+
+		TotalBytesRead += BytesRead;
+		Destination += BytesRead;
+	}
+}
+
+
+
+/**
+ * Decompresses ogg vorbis data to raw PCM data.
+ *
+ * @param	PCMData		where to place the decompressed sound
+ * @param	bLooping	whether to loop the wav by seeking to the start, or pad the buffer with zeroes
+ * @param	BufferSize	number of bytes of PCM data to create. A value of 0 means decompress the entire sound.
+ *
+ * @return	bool		true if the end of the data was reached (for both single shot and looping sounds)
+ */
+
+bool FVorbisAudioInfo::ReadCompressedData( uint8* InDestination, bool bLooping, uint32 BufferSize )
+{
+	if (!bDllLoaded)
+	{
+		UE_LOG(LogAudio, Error, TEXT("FVorbisAudioInfo::ReadCompressedData failed due to vorbis DLL not being loaded."));
+		return true;
+	}
+
+	bool		bLooped;
+	uint32		TotalBytesRead;
+
+
+
+	FScopeLock ScopeLock(&VorbisCriticalSection);
+
+	if (!bHeaderParsed)
+	{
+		UE_LOG(LogAudio, Error, TEXT("FVorbisAudioInfo::ReadCompressedData failed due to not parsing header first."));
+		return true;
+	}
+
+	bLooped = false;
+
+	// Work out number of samples to read
+	TotalBytesRead = 0;
+	char* Destination = ( char* )InDestination;
+
+	check( VFWrapper != NULL );
+
+	while( TotalBytesRead < BufferSize )
+	{
+		long BytesRead = ov_read( &VFWrapper->vf, Destination, BufferSize - TotalBytesRead, 0, 2, 1, NULL );
+		if( !BytesRead )
+		{
+			// We've reached the end
+			bLooped = true;
+			if( bLooping )
+			{
+				int Result = ov_pcm_seek_page( &VFWrapper->vf, 0 );
+				if (Result < 0)
+				{
+					// indicates an error - fill remainder of buffer with zero
+					FMemory::Memzero(Destination, BufferSize - TotalBytesRead);
+					return true;
+				}
+			}
+			else
+			{
+				int32 Count = ( BufferSize - TotalBytesRead );
+				FMemory::Memzero( Destination, Count );
+
+				BytesRead += BufferSize - TotalBytesRead;
+			}
+		}
+		else if (BytesRead < 0)
+		{
+			// indicates an error - fill remainder of buffer with zero
+			FMemory::Memzero(Destination, BufferSize - TotalBytesRead);
+			return false;
+		}
+
+		TotalBytesRead += BytesRead;
+		Destination += BytesRead;
+	}
+
+	return( bLooped );
+}
+
+void FVorbisAudioInfo::SeekToTime( const float SeekTime )
+{
+	if (!bDllLoaded)
+	{
+		UE_LOG(LogAudio, Error, TEXT("FVorbisAudioInfo::SeekToTime failed due to vorbis DLL not being loaded."));
+		return;
+	}
+
+	FScopeLock ScopeLock(&VorbisCriticalSection);
+
+	if (!bHeaderParsed)
+	{
+		UE_LOG(LogAudio, Error, TEXT("FVorbisAudioInfo::SeekToTime failed due to not parsing header first."));
+		return;
+	}
+
+	const float TargetTime = FMath::Min(SeekTime, (float)ov_time_total(&VFWrapper->vf, -1));
+	ov_time_seek( &VFWrapper->vf, TargetTime );
+}
+
+void FVorbisAudioInfo::SeekToFrame(const uint32 SeekFrames)
+{
+	if (!bDllLoaded)
+	{
+		UE_LOG(LogAudio, Error, TEXT("FVorbisAudioInfo::SeekToTime failed due to vorbis DLL not being loaded."));
+		return;
+	}
+
+	FScopeLock ScopeLock(&VorbisCriticalSection);
+
+	if (!bHeaderParsed)
+	{
+		UE_LOG(LogAudio, Error, TEXT("FVorbisAudioInfo::SeekToTime failed due to not parsing header first."));
+		return;
+	}
+
+	vorbis_info* vi = ov_info(&VFWrapper->vf, -1);
+	int64 SeekSamples = SeekFrames * vi->channels;
+	const int64 TargetSample = FMath::Min(SeekSamples, (int64)ov_pcm_total(&VFWrapper->vf, -1));
+	ov_pcm_seek(&VFWrapper->vf, TargetSample);
+}
+
+void FVorbisAudioInfo::EnableHalfRate( bool HalfRate )
+{
+	if (!bDllLoaded)
+	{
+		UE_LOG(LogAudio, Error, TEXT("FVorbisAudioInfo::EnableHalfRate failed due to vorbis DLL not being loaded."));
+		return;
+	}
+
+	FScopeLock ScopeLock(&VorbisCriticalSection);
+
+	if (!bHeaderParsed)
+	{
+		UE_LOG(LogAudio, Error, TEXT("FVorbisAudioInfo::EnableHalfRate failed due to not parsing header first."));
+		return;
+	}
+
+	ov_halfrate(&VFWrapper->vf, int32(HalfRate));
+}
+
+bool FVorbisAudioInfo::StreamCompressedInfoInternal(const FSoundWaveProxyPtr& InWaveProxy, struct FSoundQualityInfo* QualityInfo)
+{
+	if (ensure(InWaveProxy.IsValid()))
+	{
+		if (!bDllLoaded)
+		{
+			UE_LOG(LogAudio, Error, TEXT("FVorbisAudioInfo::StreamCompressedInfoInternal failed to parse header due to vorbis DLL not being loaded for sound '%s'."), *InWaveProxy->GetFName().ToString());
+			return false;
+		}
+
+	
+
+		FScopeLock ScopeLock(&VorbisCriticalSection);
+
+		if (!VFWrapper)
+		{
+			UE_LOG(LogAudio, Error, TEXT("FVorbisAudioInfo::StreamCompressedInfoInternal failed due to no vorbis wrapper for sound '%s'."), *InWaveProxy->GetFName().ToString());
+			return false;
+		}
+
+		ov_callbacks Callbacks;
+
+		SrcBufferData = NULL;
+		SrcBufferDataSize = 0;
+		BufferOffset = 0;
+
+		Callbacks.read_func = OggReadStreaming;
+		Callbacks.close_func = OggCloseStreaming;
+		Callbacks.seek_func = NULL;	// Force streaming
+		Callbacks.tell_func = NULL;	// Force streaming
+
+		bHeaderParsed = GetCompressedInfoCommon(&Callbacks, QualityInfo);
+		if (!bHeaderParsed)
+		{
+			UE_LOG(LogAudio, Error, TEXT("FVorbisAudioInfo::StreamCompressedInfoInternal failed to parse header for '%s'. LastError=%s"), 
+				*InWaveProxy->GetFName().ToString(), GetVorbisError(LastErrorCode));
+		}
+	
+
+		return bHeaderParsed;
+	}
+
+	return false;
+}
+
+int32 FVorbisAudioInfo::GetAudioDataStartOffset() const
+{
+	FScopeLock ScopeLock(&VorbisCriticalSection);
+
+	if (VFWrapper && VFWrapper->vf.dataoffsets)
+	{
+		return VFWrapper->vf.dataoffsets[0];
+	}
+	return -1;
+}
+
+
+const uint8* FVorbisAudioInfo::GetLoadedChunk(FSoundWaveProxyPtr InSoundWave, uint32 ChunkIndex, uint32& OutChunkSize)
+{
+	if (ensure(InSoundWave.IsValid()))
+	{
+		if (ChunkIndex >= InSoundWave->GetNumChunks())
+		{
+			OutChunkSize = 0;
+			return nullptr;
+		}
+		else if (ChunkIndex == 0)
+		{
+			TArrayView<const uint8> ZerothChunk = FSoundWaveProxy::GetZerothChunk(InSoundWave, true);
+			OutChunkSize = ZerothChunk.Num();
+			return ZerothChunk.GetData();
+		}
+		else
+		{
+			CurCompressedChunkHandle = IStreamingManager::Get().GetAudioStreamingManager().GetLoadedChunk(InSoundWave, ChunkIndex, false, true);
+			OutChunkSize = CurCompressedChunkHandle.Num();
+			return CurCompressedChunkHandle.GetData();
+		}
+	}
+
+	return nullptr;
+}
+
+bool FVorbisAudioInfo::StreamCompressedData(uint8* InDestination, bool bLooping, uint32 BufferSize, int32& OutNumBytesStreamed)
+{
+	const uint32 DestinationSize = BufferSize;
+	OutNumBytesStreamed = 0;
+
+	if (!bDllLoaded)
+	{
+		UE_LOG(LogAudio, Error, TEXT("FVorbisAudioInfo::StreamCompressedData failed due to vorbis DLL not being loaded."));
+		return true;
+	}
+
+	check( VFWrapper != NULL );
+
+
+	FScopeLock ScopeLock(&VorbisCriticalSection);
+
+	if (!bHeaderParsed)
+	{
+		UE_LOG(LogAudio, Error, TEXT("FVorbisAudioInfo::StreamCompressedData failed due to not parsing header first."));
+		return true;
+	}
+
+	bool bLooped = false;
+
+	while( BufferSize > 0 )
+	{
+		long BytesActuallyRead = ov_read( &VFWrapper->vf, (char*)InDestination, (int)BufferSize, 0, 2, 1, NULL );
+
+		if( BytesActuallyRead <= 0 )
+		{
+			// if we read 0 bytes or there was an error, instead of assuming we looped, lets write out zero's.
+			// this means that the chunk wasn't loaded in time.
+
+			check(StreamingSoundWave->GetNumChunks() < ((uint32)TNumericLimits<int32>::Max()));
+
+			if (NextStreamingChunkIndex < static_cast<int32>(StreamingSoundWave->GetNumChunks()))
+			{
+				// zero out the rest of the buffer
+				FMemory::Memzero(InDestination, BufferSize);
+				return false;
+			}
+
+			// We've reached the end
+			bLooped = true;
+			TimesLoopedWithoutDecompressedAudio++;
+
+			// Clean up decoder state:
+			BufferOffset = 0;
+			ov_clear(&VFWrapper->vf);
+			FMemory::Memzero(&VFWrapper->vf, sizeof(OggVorbis_File));
+
+			constexpr int32 NumReadFailiuresUntilTimeout = 64;
+			const bool bDidTimeoutOnDecompressionFailiures = (VorbisReadFailiureTimeoutCVar > 0) && (TimesLoopedWithoutDecompressedAudio >= NumReadFailiuresUntilTimeout);
+
+			UE_CLOG(bDidTimeoutOnDecompressionFailiures, LogAudio, Error, TEXT("Timed out trying to decompress a looping sound after %d attempts; aborting."), VorbisReadFailiureTimeoutCVar);
+
+			// If we're looping, then we need to make sure we wrap the stream chunks back to 0
+			if (bLooping && !bDidTimeoutOnDecompressionFailiures)
+			{
+				NextStreamingChunkIndex = 0;
+			}
+			else
+			{
+				// Here we clear out the remainder of the buffer and bail.
+				FMemory::Memzero(InDestination, BufferSize);
+				break;
+			}
+
+			// Since we can't tell a streaming file to go back to the start of the stream (there is no seek) we have to close and reopen it.
+			ov_callbacks Callbacks;
+			Callbacks.read_func = OggReadStreaming;
+			Callbacks.close_func = OggCloseStreaming;
+			Callbacks.seek_func = NULL;	// Force streaming
+			Callbacks.tell_func = NULL;	// Force streaming
+			int Result = ov_open_callbacks(this, &VFWrapper->vf, NULL, 0, Callbacks);
+			if (Result < 0)
+			{
+				LastErrorCode = Result;
+				bHasError = true;
+				UE_LOG(LogAudio, Error, TEXT("FVorbisAudioInfo::StreamCompressedData, ov_open_callbacks error code: %d : %s"), Result, GetVorbisError(Result));
+				break;
+			}
+
+			// else we start over to get the samples from the start of the compressed audio data
+			continue;
+		}
+		else
+		{
+			TimesLoopedWithoutDecompressedAudio = 0;
+		}
+
+		InDestination += BytesActuallyRead;
+		BufferSize -= BytesActuallyRead;
+	}
+
+	OutNumBytesStreamed = DestinationSize - BufferSize;
+	return( bLooped );
+}
+
+bool FVorbisAudioInfo::HasError() const
+{
+	return Super::HasError() | bHasError;
+}
+
+void LoadVorbisLibraries()
+{
+	static bool bIsInitialized = false;
+	if (!bIsInitialized)
+	{
+		bIsInitialized = true;
+#if WITH_OGGVORBIS_DLL && WITH_OGGVORBIS
+		//@todo if ogg is every ported to another platform, then use the platform abstraction to load these DLLs
+		// Load the Ogg dlls
+#  if _MSC_VER >= 1900
+		FString VSVersion = TEXT("VS2015/");
+#  else
+#    error "Unsupported Visual Studio version."
+#  endif
+		FString PlatformString = TEXT("Win32");
+		FString DLLNameStub = TEXT(".dll");
+#if PLATFORM_64BITS
+		PlatformString = TEXT("Win64");
+		DLLNameStub = TEXT("_64.dll");
+#endif
+#if PLATFORM_HOLOLENS
+		PlatformString = TEXT("HoloLens");
+#endif
+
+#if PLATFORM_CPU_ARM_FAMILY && !defined(_M_ARM64EC)
+#if PLATFORM_64BITS
+		FString RootOggPath = FPaths::EngineDir() / TEXT("Binaries/ThirdParty/Ogg/") / PlatformString / VSVersion / TEXT("arm64/");
+		FString RootVorbisPath = FPaths::EngineDir() / TEXT("Binaries/ThirdParty/Vorbis/") / PlatformString / VSVersion / TEXT("arm64/");
+#else
+		FString RootOggPath = FPaths::EngineDir() / TEXT("Binaries/ThirdParty/Ogg/") / PlatformString / VSVersion / TEXT("arm/");
+		FString RootVorbisPath = FPaths::EngineDir() / TEXT("Binaries/ThirdParty/Vorbis/") / PlatformString / VSVersion / TEXT("arm/");
+#endif
+#else
+		FString RootOggPath = FPaths::EngineDir() / TEXT("Binaries/ThirdParty/Ogg/") / PlatformString / VSVersion;
+		FString RootVorbisPath = FPaths::EngineDir() / TEXT("Binaries/ThirdParty/Vorbis/") / PlatformString / VSVersion;
+#endif
+
+		FString DLLToLoad = RootOggPath + TEXT("libogg") + DLLNameStub;
+		void* LibOggHandle = FPlatformProcess::GetDllHandle(*DLLToLoad);
+		verifyf(LibOggHandle, TEXT("Failed to load DLL %s"), *DLLToLoad);
+		// Load the Vorbis dlls
+		DLLToLoad = RootVorbisPath + TEXT("libvorbis") + DLLNameStub;
+
+		void* LibVorbisHandle = FPlatformProcess::GetDllHandle(*DLLToLoad);
+		verifyf(LibVorbisHandle, TEXT("Failed to load DLL %s"), *DLLToLoad);
+		DLLToLoad = RootVorbisPath + TEXT("libvorbisfile") + DLLNameStub;
+		
+		void* LibVorbisFileHandle = FPlatformProcess::GetDllHandle(*DLLToLoad);
+		verifyf(LibVorbisFileHandle, TEXT("Failed to load DLL %s"), *DLLToLoad);
+
+		// Set that we successfully loaded everything so we can do nothing if it fails and avoid a crash
+		bDllLoaded = LibOggHandle && LibVorbisHandle && LibVorbisFileHandle;
+
+		if (!bDllLoaded)
+		{
+			UE_LOG(LogAudio, Error, TEXT("Failed to load lib vorbis libraries."));
+		}
+		else
+		{
+			UE_LOG(LogAudioDebug, Display, TEXT("Lib vorbis DLL was dynamically loaded."));
+		}
+#elif WITH_OGGVORBIS
+		bDllLoaded = true;
+#endif	//WITH_OGGVORBIS_DLL
+	}
+}
+
+#endif		// WITH_OGGVORBIS
+
+class VORBISAUDIODECODER_API FVorbisAudioDecoderModule : public IModuleInterface
+{
+public:	
+	TUniquePtr<IAudioInfoFactory> Factory;
+
+	virtual void StartupModule() override
+	{
+		LoadVorbisLibraries();
+		Factory = MakeUnique<FSimpleAudioInfoFactory>([] { return new FVorbisAudioInfo(); }, Audio::NAME_OGG);
+	}
+
+	virtual void ShutdownModule() override {}
+};
+
+IMPLEMENT_MODULE(FVorbisAudioDecoderModule, VorbisAudioDecoder)
diff --git a/Engine/Source/Runtime/VorbisAudioDecoder/Module/Public/Decoders/VorbisAudioInfo.h b/Engine/Source/Runtime/VorbisAudioDecoder/Module/Public/Decoders/VorbisAudioInfo.h
--- a/Engine/Source/Runtime/VorbisAudioDecoder/Module/Public/Decoders/VorbisAudioInfo.h
+++ b/Engine/Source/Runtime/VorbisAudioDecoder/Module/Public/Decoders/VorbisAudioInfo.h
@@ -27,14 +27,13 @@
 
 namespace VorbisChannelInfo
 {
-	extern ENGINE_API const int32 Order[8][8];
+	extern VORBISAUDIODECODER_API const int32 Order[8][8];
 }
 
 /**
  * Loads vorbis dlls
 */
-ENGINE_API void LoadVorbisLibraries();
-
+VORBISAUDIODECODER_API void LoadVorbisLibraries();
 #if WITH_OGGVORBIS
 
 /** 
@@ -43,9 +42,8 @@ ENGINE_API void LoadVorbisLibraries();
 class FVorbisAudioInfo : public ICompressedAudioInfo
 {
 public:
-	ENGINE_API FVorbisAudioInfo( void );
-	ENGINE_API virtual ~FVorbisAudioInfo( void );
-
+	VORBISAUDIODECODER_API FVorbisAudioInfo( void );
+	VORBISAUDIODECODER_API virtual ~FVorbisAudioInfo( void );
 	/** Emulate read from memory functionality */
 	size_t			ReadMemory( void *ptr, uint32 size );
 	int				SeekMemory( uint32 offset, int whence );
@@ -65,8 +63,7 @@ public:
 	 * 
 	 * @param	Resource		Info about vorbis data
 	 */
-	ENGINE_API virtual bool ReadCompressedInfo( const uint8* InSrcBufferData, uint32 InSrcBufferDataSize, struct FSoundQualityInfo* QualityInfo ) override;
-
+	VORBISAUDIODECODER_API virtual bool ReadCompressedInfo( const uint8* InSrcBufferData, uint32 InSrcBufferDataSize, struct FSoundQualityInfo* QualityInfo ) override;
 	/** 
 	 * Decompresses ogg data to raw PCM data. 
 	 * 
@@ -76,24 +73,19 @@ public:
 	 *
 	 * @return	bool		true if the end of the data was reached (for both single shot and looping sounds)
 	 */
-	ENGINE_API virtual bool ReadCompressedData( uint8* InDestination, bool bLooping, uint32 BufferSize ) override;
-
-	ENGINE_API virtual void SeekToTime( const float SeekTime ) override;
-
-	ENGINE_API virtual void SeekToFrame(const uint32 SeekFrame) override;
-
+	VORBISAUDIODECODER_API virtual bool ReadCompressedData( uint8* InDestination, bool bLooping, uint32 BufferSize ) override;
+	VORBISAUDIODECODER_API virtual void SeekToTime( const float SeekTime ) override;
+	VORBISAUDIODECODER_API virtual void SeekToFrame(const uint32 SeekFrame) override;
 	/** 
 	 * Decompress an entire ogg data file to a TArray
 	 */
-	ENGINE_API virtual void ExpandFile( uint8* DstBuffer, struct FSoundQualityInfo* QualityInfo ) override;
-
+	VORBISAUDIODECODER_API virtual void ExpandFile( uint8* DstBuffer, struct FSoundQualityInfo* QualityInfo ) override;
 	/** 
 	 * Sets ogg to decode to half-rate
 	 * 
 	 * @param	Resource		Info about vorbis data
 	 */
-	ENGINE_API virtual void EnableHalfRate( bool HalfRate ) override;
-
+	VORBISAUDIODECODER_API virtual void EnableHalfRate( bool HalfRate ) override;
 	virtual uint32 GetSourceBufferSize() const override { return SrcBufferDataSize; }
 
 	virtual bool UsesVorbisChannelOrdering() const override { return true; }
@@ -111,8 +103,7 @@ protected:
 	virtual bool StreamCompressedInfoInternal(const FSoundWaveProxyPtr& InWaveProxy, struct FSoundQualityInfo* QualityInfo) override;
 
 	friend class FAudioFormatOgg;
-	ENGINE_API int32 GetAudioDataStartOffset() const;
-
+	VORBISAUDIODECODER_API int32 GetAudioDataStartOffset() const;
 private:
 	using Super = ICompressedAudioInfo; 
 	const uint8* GetLoadedChunk(FSoundWaveProxyPtr InSoundWave, uint32 ChunkIndex, uint32& OutChunkSize);
diff --git a/Engine/Source/Runtime/VorbisAudioDecoder/Module/VorbisAudioDecoder.Build.cs b/Engine/Source/Runtime/VorbisAudioDecoder/Module/VorbisAudioDecoder.Build.cs
new file mode 100644
--- /dev/null
+++ b/Engine/Source/Runtime/VorbisAudioDecoder/Module/VorbisAudioDecoder.Build.cs
@@ -0,0 +1,29 @@
+// Copyright Epic Games, Inc. All Rights Reserved.
+
+using UnrealBuildTool;
+
+public class VorbisAudioDecoder : ModuleRules
+{
+    public VorbisAudioDecoder(ReadOnlyTargetRules Target) : base(Target)
+    {
+        ShortName = "VorbisAudioDecoder";
+
+        PublicDependencyModuleNames.AddRange(
+            new string[]
+            {
+                "Core",
+                "CoreUObject",
+                "Engine"
+            }
+        );
+        
+        AddEngineThirdPartyPrivateStaticDependencies(
+	        Target,
+		        "UEOgg",
+		        "Vorbis",
+		        "VorbisFile"
+        );        
+
+	    PublicDefinitions.Add("WITH_OGGVORBIS=1");
+    }
+}
