Fix for Client Crash - ADPCM::DecodeBlock() - Code by Aaron.

diff --git a/Engine/Plugins/Runtime/Metasound/Source/MetasoundEngine/Private/MetasoundWavePlayerNode.cpp b/Engine/Plugins/Runtime/Metasound/Source/MetasoundEngine/Private/MetasoundWavePlayerNode.cpp
--- a/Engine/Plugins/Runtime/Metasound/Source/MetasoundEngine/Private/MetasoundWavePlayerNode.cpp
+++ b/Engine/Plugins/Runtime/Metasound/Source/MetasoundEngine/Private/MetasoundWavePlayerNode.cpp
@@ -795,6 +795,7 @@ namespace Metasound
 		 *
 		 * @param OutBuffer - Buffer to place generated audio.
 		 * @param OutSourceState - Source state for tracking state of OutBuffer.
+		 * Returns true if succesfully generated audio. 
 		 */
 		void GenerateSourceAudio(Audio::FMultichannelCircularBuffer& OutBuffer, WavePlayerNodePrivate::FSourceBufferState& OutSourceState)
 		{
@@ -806,7 +807,16 @@ namespace Metasound
 				const int32 NumSamplesToGenerate = DeinterleaveBufferBlockSizeInFrames * WaveProxyReader->GetNumChannels();
  				check(NumSamplesToGenerate == InterleavedBuffer.Num())
 
-				WaveProxyReader->PopAudio(InterleavedBuffer);
+				// if the wave proxy reader has failed, write out silence. 
+				if (WaveProxyReader->HasFailed())
+				{
+					FMemory::Memset(InterleavedBuffer.GetData(), 0, sizeof(float) * InterleavedBuffer.Num());
+				}
+				else
+				{
+					WaveProxyReader->PopAudio(InterleavedBuffer);
+				}
+
 				ConvertDeinterleave->ProcessAudio(InterleavedBuffer, DeinterleavedBuffer);
 
 				for (int32 ChannelIndex = 0; ChannelIndex < NumDeinterleaveChannels; ChannelIndex++)
diff --git a/Engine/Source/Runtime/AdpcmAudioDecoder/Module/Private/ADPCMAudioInfo.cpp b/Engine/Source/Runtime/AdpcmAudioDecoder/Module/Private/ADPCMAudioInfo.cpp
--- a/Engine/Source/Runtime/AdpcmAudioDecoder/Module/Private/ADPCMAudioInfo.cpp
+++ b/Engine/Source/Runtime/AdpcmAudioDecoder/Module/Private/ADPCMAudioInfo.cpp
@@ -44,8 +44,9 @@ FAutoConsoleVariableRef CVarChanceForIntentionalChunkMiss(
 
 namespace ADPCM
 {
-	void DecodeBlock(const uint8* EncodedADPCMBlock, int32 BlockSize, int16* DecodedPCMData);
-	void DecodeBlockStereo(const uint8* EncodedADPCMBlockLeft, const uint8* EncodedADPCMBlockRight, int32 BlockSize, int16* DecodedPCMData);
+	// Returns true if successful, false if there was an error in decoding
+	bool DecodeBlock(const uint8* EncodedADPCMBlock, int32 BlockSize, int16* DecodedPCMData);
+	bool DecodeBlockStereo(const uint8* EncodedADPCMBlockLeft, const uint8* EncodedADPCMBlockRight, int32 BlockSize, int16* DecodedPCMData);
 }
 
 FADPCMAudioInfo::FADPCMAudioInfo(void)
@@ -353,10 +354,17 @@ bool FADPCMAudioInfo::ReadCompressedData(uint8* Destination, bool bLooping, uint
 				// Decompress one block for each channel and store it in UncompressedBlockData
 				for(int32 ChannelItr = 0; ChannelItr < NumChannels; ++ChannelItr)
 				{
-					ADPCM::DecodeBlock(
+					bool bSuccess = ADPCM::DecodeBlock(
 						WaveInfo.SampleDataStart + (ChannelItr * TotalCompressedBlocksPerChannel + CurrentCompressedBlockIndex) * CompressedBlockSize,
 						CompressedBlockSize,
 						(int16*)(UncompressedBlockData + ChannelItr * UncompressedBlockSize));
+
+					if (!bSuccess)
+					{
+						FMemory::Memzero(Destination, BufferSize);
+						bHasError = true;
+						return true;
+					}
 				}
 
 				// Update some bookkeeping
@@ -800,10 +808,17 @@ bool FADPCMAudioInfo::StreamCompressedData(uint8* Destination, bool bLooping, ui
 				// Decompress one block for each channel and store it in UncompressedBlockData
 				for(int32 ChannelItr = 0; ChannelItr < NumChannels; ++ChannelItr)
 				{
-					ADPCM::DecodeBlock(
+					bool bSuccess = ADPCM::DecodeBlock(
 						CurCompressedChunkData + CurrentChunkBufferOffset + ChannelItr * CompressedBlockSize,
 						CompressedBlockSize,
 						(int16*)(UncompressedBlockData + ChannelItr * UncompressedBlockSize));
+
+					if (!bSuccess)
+					{
+						FMemory::Memzero(Destination, BufferSize);
+						bHasError = true;
+						return true;
+					}
 				}
 
 				// Update some bookkeeping
@@ -1206,18 +1221,27 @@ namespace ADPCM
 		return Context.Sample1;
 	}
 
-	void DecodeBlock(const uint8* EncodedADPCMBlock, int32 BlockSize, int16* DecodedPCMData)
+	bool DecodeBlock(const uint8* EncodedADPCMBlock, int32 BlockSize, int16* DecodedPCMData)
 	{
 		FAdaptationContext Context;
 		int32 ReadIndex = 0;
 		int32 WriteIndex = 0;
 
 		uint8 CoefficientIndex = ReadFromByteStream<uint8>(EncodedADPCMBlock, ReadIndex);
-		Context.AdaptationDelta = ReadFromByteStream<int16>(EncodedADPCMBlock, ReadIndex);
-		Context.Sample1 = ReadFromByteStream<int16>(EncodedADPCMBlock, ReadIndex);
-		Context.Sample2 = ReadFromByteStream<int16>(EncodedADPCMBlock, ReadIndex);
-		Context.Coefficient1 = Context.AdaptationCoefficient1[CoefficientIndex];
-		Context.Coefficient2 = Context.AdaptationCoefficient2[CoefficientIndex];
+		if (CoefficientIndex >= NUM_ADAPTATION_COEFF)
+		{
+			UE_LOG(LogAudio, Error, TEXT("Decoding ADPCM block resulted in bad CoefficientIndex (%d). BlockSize: %d, ReadIndex: %d"), CoefficientIndex, BlockSize, ReadIndex);
+			return false;
+		}
+		else
+		{
+			Context.AdaptationDelta = ReadFromByteStream<int16>(EncodedADPCMBlock, ReadIndex);
+			Context.Sample1 = ReadFromByteStream<int16>(EncodedADPCMBlock, ReadIndex);
+			Context.Sample2 = ReadFromByteStream<int16>(EncodedADPCMBlock, ReadIndex);
+
+			Context.Coefficient1 = Context.AdaptationCoefficient1[CoefficientIndex];
+			Context.Coefficient2 = Context.AdaptationCoefficient2[CoefficientIndex];
+		}
 
 		// The first two samples are sent directly to the output in reverse order, as per the standard
 		DecodedPCMData[WriteIndex++] = Context.Sample2;
@@ -1236,10 +1260,12 @@ namespace ADPCM
 			EncodedNibble = EncodedNibblePair & 0x0F;
 			DecodedPCMData[WriteIndex++] = DecodeNibble(Context, EncodedNibble);
 		}
+
+		return true;
 	}
 
 	// Decode two PCM streams and interleave as stereo data
-	void DecodeBlockStereo(const uint8* EncodedADPCMBlockLeft, const uint8* EncodedADPCMBlockRight, int32 BlockSize, int16* DecodedPCMData)
+	bool DecodeBlockStereo(const uint8* EncodedADPCMBlockLeft, const uint8* EncodedADPCMBlockRight, int32 BlockSize, int16* DecodedPCMData)
 	{
 		FAdaptationContext ContextLeft;
 		FAdaptationContext ContextRight;
@@ -1248,6 +1274,13 @@ namespace ADPCM
 		int32 WriteIndex = 0;
 
 		uint8 CoefficientIndexLeft = ReadFromByteStream<uint8>(EncodedADPCMBlockLeft, ReadIndexLeft);
+
+		if (CoefficientIndexLeft >= NUM_ADAPTATION_COEFF)
+		{
+			UE_LOG(LogAudio, Error, TEXT("Decoding ADPCM block resulted in bad CoefficientIndexLeft (%d). BlockSize: %d, ReadIndexLeft: %d"), CoefficientIndexLeft, BlockSize, ReadIndexLeft);
+			return false;
+		}
+
 		ContextLeft.AdaptationDelta = ReadFromByteStream<int16>(EncodedADPCMBlockLeft, ReadIndexLeft);
 		ContextLeft.Sample1 = ReadFromByteStream<int16>(EncodedADPCMBlockLeft, ReadIndexLeft);
 		ContextLeft.Sample2 = ReadFromByteStream<int16>(EncodedADPCMBlockLeft, ReadIndexLeft);
@@ -1255,6 +1288,13 @@ namespace ADPCM
 		ContextLeft.Coefficient2 = ContextLeft.AdaptationCoefficient2[CoefficientIndexLeft];
 
 		uint8 CoefficientIndexRight = ReadFromByteStream<uint8>(EncodedADPCMBlockRight, ReadIndexRight);
+
+		if (CoefficientIndexRight >= NUM_ADAPTATION_COEFF)
+		{
+			UE_LOG(LogAudio, Error, TEXT("Decoding ADPCM block resulted in bad CoefficientIndexRight (%d). BlockSize: %d, ReadIndexLeft: %d"), CoefficientIndexRight, BlockSize, ReadIndexRight);
+			return false;
+		}
+
 		ContextRight.AdaptationDelta = ReadFromByteStream<int16>(EncodedADPCMBlockRight, ReadIndexRight);
 		ContextRight.Sample1 = ReadFromByteStream<int16>(EncodedADPCMBlockRight, ReadIndexRight);
 		ContextRight.Sample2 = ReadFromByteStream<int16>(EncodedADPCMBlockRight, ReadIndexRight);
@@ -1291,6 +1331,7 @@ namespace ADPCM
 			DecodedPCMData[WriteIndex++] = DecodeNibble(ContextRight, EncodedNibbleRight);
 		}
 
+		return true;
 	}
 } // end namespace ADPCM
 
@@ -1311,4 +1352,4 @@ public:
 	virtual void ShutdownModule() override {}
 };
 
-IMPLEMENT_MODULE(FAdpcmAudioDecoderModule, AdpcmAudioDecoder)
+IMPLEMENT_MODULE(FAdpcmAudioDecoderModule, AdpcmAudioDecoder)
\ No newline at end of file
diff --git a/Engine/Source/Runtime/Engine/Classes/Sound/SoundWaveProxyReader.h b/Engine/Source/Runtime/Engine/Classes/Sound/SoundWaveProxyReader.h
--- a/Engine/Source/Runtime/Engine/Classes/Sound/SoundWaveProxyReader.h
+++ b/Engine/Source/Runtime/Engine/Classes/Sound/SoundWaveProxyReader.h
@@ -154,6 +154,15 @@ public:
 	 */
 	ENGINE_API bool CanProduceMoreAudio() const;
 
+	/** Returns TRUE if the reader has encountered a decoder failure.
+	 *
+	 * @return true if there is a failure, false if no failure.
+	 */
+	ENGINE_API bool HasFailed() const
+	{ 
+		return DecodeResult == EDecodeResult::Fail; 
+	}
+
 private:
 
 	enum EDecodeResult
diff --git a/Engine/Source/Runtime/Engine/Private/AudioDevice.cpp b/Engine/Source/Runtime/Engine/Private/AudioDevice.cpp
--- a/Engine/Source/Runtime/Engine/Private/AudioDevice.cpp
+++ b/Engine/Source/Runtime/Engine/Private/AudioDevice.cpp
@@ -4476,6 +4476,7 @@ void FAudioDevice::StartSources(TArray<FWaveInstance*>& WaveInstances, int32 Fir
 					// This can happen if e.g. the USoundWave pointed to by the WaveInstance is not a valid sound file.
 					// If we don't stop the wave file, it will continue to try initializing the file every frame, which is a perf hit
 					UE_LOG(LogAudio, Log, TEXT("Failed to start sound source for %s"), (WaveInstance->ActiveSound && WaveInstance->ActiveSound->Sound) ? *WaveInstance->ActiveSound->Sound->GetName() : TEXT("UNKNOWN") );
+					AddSoundToStop(WaveInstance->ActiveSound);
 					WaveInstance->StopWithoutNotification();
 					Source->WaveInstance = nullptr;
 					FreeSources.Add(Source);
