ADPCMAudioInfo update and add logs

diff --git a/Engine/Source/Runtime/AdpcmAudioDecoder/Module/Private/ADPCMAudioInfo.cpp b/Engine/Source/Runtime/AdpcmAudioDecoder/Module/Private/ADPCMAudioInfo.cpp
--- a/Engine/Source/Runtime/AdpcmAudioDecoder/Module/Private/ADPCMAudioInfo.cpp
+++ b/Engine/Source/Runtime/AdpcmAudioDecoder/Module/Private/ADPCMAudioInfo.cpp
@@ -245,7 +245,7 @@ bool FADPCMAudioInfo::ReadCompressedInfo(const uint8* InSrcBufferData, uint32 In
 	if (!InSrcBufferData)
 	{
 		FString Name = QualityInfo ? QualityInfo->DebugName : TEXT("Unknown");
-		UE_LOG(LogAudio, Warning, TEXT("Failed to read compressed ADPCM audio from ('%s') because there was no resource data."), *Name);
+		UE_LOG(LogAudio, Error, TEXT("FADPCMAudioInfo::ReadCompressedInfo: Failed to read compressed ADPCM audio from ('%s') because there was no resource data."), *Name);
 
 		return false;
 	}
@@ -257,7 +257,7 @@ bool FADPCMAudioInfo::ReadCompressedInfo(const uint8* InSrcBufferData, uint32 In
 
 	if (!WaveInfo.ReadWaveInfo((uint8*)SrcBufferData, SrcBufferDataSize, nullptr, false, &FormatHeader))
 	{
-		UE_LOG(LogAudio, Warning, TEXT("WaveInfo.ReadWaveInfo Failed"));
+		UE_LOG(LogAudio, Error, TEXT("FADPCMAudioInfo::ReadCompressedInfo: WaveInfo.ReadWaveInfo Failed"));
 		return false;
 	}
 
@@ -303,6 +303,7 @@ bool FADPCMAudioInfo::ReadCompressedInfo(const uint8* InSrcBufferData, uint32 In
 	}
 	else
 	{
+		UE_LOG(LogAudio, Error, TEXT("FADPCMAudioInfo::ReadCompressedInfo: Invalid format. Expected %d or %d. Got %d"), WAVE_FORMAT_LPCM, WAVE_FORMAT_ADPCM, Format);
 		return false;
 	}
 
