Avoid non-determinism in Audio DDC data by ensuring that all structure fields are initialized and padding bytes in structures are zero'd before use.  Not being accompanied by a DDC version bump for audio to avoid incurring patch bloat.

diff --git a/Engine/Source/Developer/AudioFormatADPCM/Private/AudioFormatADPCM.cpp b/Engine/Source/Developer/AudioFormatADPCM/Private/AudioFormatADPCM.cpp
--- a/Engine/Source/Developer/AudioFormatADPCM/Private/AudioFormatADPCM.cpp
+++ b/Engine/Source/Developer/AudioFormatADPCM/Private/AudioFormatADPCM.cpp
@@ -157,6 +157,7 @@ namespace LPCM
 		Format.nAvgBytesPerSec = Format.nBlockAlign * QualityInfo.SampleRate;
 		Format.wBitsPerSample = 16;
 		Format.wFormatTag = WAVE_FORMAT_LPCM;
+		Format.cbSize = 0;
 
 		RiffDataChunk RiffDataChunks[2];
 		RiffDataChunks[0].ID = UE_MAKEFOURCC('f','m','t',' ');
diff --git a/Engine/Source/Runtime/AdpcmAudioDecoder/Module/Public/Decoders/ADPCMAudioInfo.h b/Engine/Source/Runtime/AdpcmAudioDecoder/Module/Public/Decoders/ADPCMAudioInfo.h
--- a/Engine/Source/Runtime/AdpcmAudioDecoder/Module/Public/Decoders/ADPCMAudioInfo.h
+++ b/Engine/Source/Runtime/AdpcmAudioDecoder/Module/Public/Decoders/ADPCMAudioInfo.h
@@ -70,6 +70,7 @@ namespace ADPCM
 
 		ADPCMFormatHeader()
 		{
+			FMemory::Memzero(this, sizeof(ADPCMFormatHeader));
 			int16 AdaptationCoefficient1[NUM_ADAPTATION_COEFF];
 			int16 AdaptationCoefficient2[NUM_ADAPTATION_COEFF];
 			GetAdaptationCoefficients(AdaptationCoefficient1, AdaptationCoefficient2);
