ADPCM rework (working)

diff --git a/Engine/Source/Runtime/AdpcmAudioDecoder/Module/Private/ADPCMAudioInfo.cpp b/Engine/Source/Runtime/AdpcmAudioDecoder/Module/Private/ADPCMAudioInfo.cpp
--- a/Engine/Source/Runtime/AdpcmAudioDecoder/Module/Private/ADPCMAudioInfo.cpp
+++ b/Engine/Source/Runtime/AdpcmAudioDecoder/Module/Private/ADPCMAudioInfo.cpp
@@ -3,6 +3,7 @@
 
 #include "Decoders/ADPCMAudioInfo.h"
 #include "Interfaces/IAudioFormat.h"
+#include "Adpcm.h"
 
 static int32 bDisableADPCMSeekingCVar = 0;
 FAutoConsoleVariableRef CVarDisableADPCMSeeking(
@@ -39,16 +40,6 @@ FAutoConsoleVariableRef CVarChanceForIntentionalChunkMiss(
 	TEXT("If this is set > 0 we will intentionally drop chunks. Used for debugging..\n"),
 	ECVF_Default);
 
-#define WAVE_FORMAT_LPCM  1
-#define WAVE_FORMAT_ADPCM 2
-
-namespace ADPCM
-{
-	// Returns true if successful, false if there was an error in decoding
-	bool DecodeBlock(const uint8* EncodedADPCMBlock, int32 BlockSize, int16* DecodedPCMData);
-	bool DecodeBlockStereo(const uint8* EncodedADPCMBlockLeft, const uint8* EncodedADPCMBlockRight, int32 BlockSize, int16* DecodedPCMData);
-}
-
 FADPCMAudioInfo::FADPCMAudioInfo(void)
 	: NumConsecutiveReadFailiures(0)
 	, UncompressedBlockSize(0)
@@ -1054,303 +1045,3 @@ const uint8* FADPCMAudioInfo::GetLoadedChunk(const FSoundWaveProxyPtr& InSoundWa
 		return CurCompressedChunkHandle.GetData();
 	}
 }
-
-//////////////////////////////////////////////////////////////////////////
-// Copied from IOS - probably want to split and share
-//////////////////////////////////////////////////////////////////////////
-#define NUM_ADAPTATION_TABLE 16
-#define NUM_ADAPTATION_COEFF 7
-#define SOUND_SOURCE_FREE 0
-#define SOUND_SOURCE_LOCKED 1
-
-namespace
-{
-	template <typename T, uint32 B>
-	inline T SignExtend(const T ValueToExtend)
-	{
-		struct { T ExtendedValue : B; } SignExtender;
-		return SignExtender.ExtendedValue = ValueToExtend;
-	}
-
-	template <typename T>
-	inline T ReadFromByteStream(const uint8* ByteStream, int32& ReadIndex, bool bLittleEndian = true)
-	{
-		T ValueRaw = 0;
-
-		if (bLittleEndian)
-		{
-#if PLATFORM_LITTLE_ENDIAN
-			for (int32 ByteIndex = 0; ByteIndex < sizeof(T); ++ByteIndex)
-#else
-			for (int32 ByteIndex = sizeof(T) - 1; ByteIndex >= 0; --ByteIndex)
-#endif // PLATFORM_LITTLE_ENDIAN
-			{
-				ValueRaw |= ByteStream[ReadIndex++] << 8 * ByteIndex;
-			}
-		}
-		else
-		{
-#if PLATFORM_LITTLE_ENDIAN
-			for (int32 ByteIndex = sizeof(T) - 1; ByteIndex >= 0; --ByteIndex)
-#else
-			for (int32 ByteIndex = 0; ByteIndex < sizeof(T); ++ByteIndex)
-#endif // PLATFORM_LITTLE_ENDIAN
-			{
-				ValueRaw |= ByteStream[ReadIndex++] << 8 * ByteIndex;
-			}
-		}
-
-		return ValueRaw;
-	}
-
-	template <typename T>
-	inline void WriteToByteStream(T Value, uint8* ByteStream, int32& WriteIndex, bool bLittleEndian = true)
-	{
-		if (bLittleEndian)
-		{
-#if PLATFORM_LITTLE_ENDIAN
-			for (int32 ByteIndex = 0; ByteIndex < sizeof(T); ++ByteIndex)
-#else
-			for (int32 ByteIndex = sizeof(T) - 1; ByteIndex >= 0; --ByteIndex)
-#endif // PLATFORM_LITTLE_ENDIAN
-			{
-				ByteStream[WriteIndex++] = (Value >> (8 * ByteIndex)) & 0xFF;
-			}
-		}
-		else
-		{
-#if PLATFORM_LITTLE_ENDIAN
-			for (int32 ByteIndex = sizeof(T) - 1; ByteIndex >= 0; --ByteIndex)
-#else
-			for (int32 ByteIndex = 0; ByteIndex < sizeof(T); ++ByteIndex)
-#endif // PLATFORM_LITTLE_ENDIAN
-			{
-				ByteStream[WriteIndex++] = (Value >> (8 * ByteIndex)) & 0xFF;
-			}
-		}
-	}
-
-	template <typename T>
-	inline T ReadFromArray(const T* ElementArray, int32& ReadIndex, int32 NumElements, int32 IndexStride = 1)
-	{
-		T OutputValue = 0;
-
-		if (ReadIndex >= 0 && ReadIndex < NumElements)
-		{
-			OutputValue = ElementArray[ReadIndex];
-			ReadIndex += IndexStride;
-		}
-
-		return OutputValue;
-	}
-
-	inline bool LockSourceChannel(int32* ChannelLock)
-	{
-		check(ChannelLock != NULL);
-		return FPlatformAtomics::InterlockedCompareExchange(ChannelLock, SOUND_SOURCE_LOCKED, SOUND_SOURCE_FREE) == SOUND_SOURCE_FREE;
-	}
-
-	inline void UnlockSourceChannel(int32* ChannelLock)
-	{
-		check(ChannelLock != NULL);
-		int32 Result = FPlatformAtomics::InterlockedCompareExchange(ChannelLock, SOUND_SOURCE_FREE, SOUND_SOURCE_LOCKED);
-
-		check(Result == SOUND_SOURCE_LOCKED);
-	}
-
-} // end namespace
-
-
-//////////////////////////////////////////////////////////////////////////
-// End of copied section
-//////////////////////////////////////////////////////////////////////////
-
-
-namespace ADPCM
-{
-	template <typename T>
-	static void GetAdaptationTable(T(&OutAdaptationTable)[NUM_ADAPTATION_TABLE])
-	{
-		// Magic values as specified by standard
-		static T AdaptationTable[] =
-		{
-			230, 230, 230, 230, 307, 409, 512, 614,
-			768, 614, 512, 409, 307, 230, 230, 230
-		};
-
-		FMemory::Memcpy(&OutAdaptationTable, AdaptationTable, sizeof(AdaptationTable));
-	}
-
-	struct FAdaptationContext
-	{
-	public:
-		// Adaptation constants
-		int32 AdaptationTable[NUM_ADAPTATION_TABLE];
-		int32 AdaptationCoefficient1[NUM_ADAPTATION_COEFF];
-		int32 AdaptationCoefficient2[NUM_ADAPTATION_COEFF];
-
-		int32 AdaptationDelta;
-		int32 Coefficient1;
-		int32 Coefficient2;
-		int32 Sample1;
-		int32 Sample2;
-
-		FAdaptationContext() :
-			AdaptationDelta(0),
-			Coefficient1(0),
-			Coefficient2(0),
-			Sample1(0),
-			Sample2(0)
-		{
-			GetAdaptationTable(AdaptationTable);
-			GetAdaptationCoefficients(AdaptationCoefficient1, AdaptationCoefficient2);
-		}
-	};
-
-	FORCEINLINE int16 DecodeNibble(FAdaptationContext& Context, uint8 EncodedNibble)
-	{
-		int32 PredictedSample = (Context.Sample1 * Context.Coefficient1 + Context.Sample2 * Context.Coefficient2) / 256;
-		PredictedSample += SignExtend<int8, 4>(EncodedNibble) * Context.AdaptationDelta;
-		PredictedSample = FMath::Clamp(PredictedSample, -32768, 32767);
-
-		// Shuffle samples for the next iteration
-		Context.Sample2 = Context.Sample1;
-		Context.Sample1 = static_cast<int16>(PredictedSample);
-		Context.AdaptationDelta = (Context.AdaptationDelta * Context.AdaptationTable[EncodedNibble]) / 256;
-		Context.AdaptationDelta = FMath::Max(Context.AdaptationDelta, 16);
-
-		return Context.Sample1;
-	}
-
-	bool DecodeBlock(const uint8* EncodedADPCMBlock, int32 BlockSize, int16* DecodedPCMData)
-	{
-		FAdaptationContext Context;
-		int32 ReadIndex = 0;
-		int32 WriteIndex = 0;
-
-		uint8 CoefficientIndex = ReadFromByteStream<uint8>(EncodedADPCMBlock, ReadIndex);
-		if (CoefficientIndex >= NUM_ADAPTATION_COEFF)
-		{
-			UE_LOG(LogAudio, Error, TEXT("Decoding ADPCM block resulted in bad CoefficientIndex (%d). BlockSize: %d, ReadIndex: %d"), CoefficientIndex, BlockSize, ReadIndex);
-			return false;
-		}
-		else
-		{
-			Context.AdaptationDelta = ReadFromByteStream<int16>(EncodedADPCMBlock, ReadIndex);
-			Context.Sample1 = ReadFromByteStream<int16>(EncodedADPCMBlock, ReadIndex);
-			Context.Sample2 = ReadFromByteStream<int16>(EncodedADPCMBlock, ReadIndex);
-
-			Context.Coefficient1 = Context.AdaptationCoefficient1[CoefficientIndex];
-			Context.Coefficient2 = Context.AdaptationCoefficient2[CoefficientIndex];
-		}
-
-		// The first two samples are sent directly to the output in reverse order, as per the standard
-		DecodedPCMData[WriteIndex++] = Context.Sample2;
-		DecodedPCMData[WriteIndex++] = Context.Sample1;
-
-		uint8 EncodedNibblePair = 0;
-		uint8 EncodedNibble = 0;
-		while (ReadIndex < BlockSize)
-		{
-			// Read from the byte stream and advance the read head.
-			EncodedNibblePair = ReadFromByteStream<uint8>(EncodedADPCMBlock, ReadIndex);
-
-			EncodedNibble = (EncodedNibblePair >> 4) & 0x0F;
-			DecodedPCMData[WriteIndex++] = DecodeNibble(Context, EncodedNibble);
-
-			EncodedNibble = EncodedNibblePair & 0x0F;
-			DecodedPCMData[WriteIndex++] = DecodeNibble(Context, EncodedNibble);
-		}
-
-		return true;
-	}
-
-	// Decode two PCM streams and interleave as stereo data
-	bool DecodeBlockStereo(const uint8* EncodedADPCMBlockLeft, const uint8* EncodedADPCMBlockRight, int32 BlockSize, int16* DecodedPCMData)
-	{
-		FAdaptationContext ContextLeft;
-		FAdaptationContext ContextRight;
-		int32 ReadIndexLeft = 0;
-		int32 ReadIndexRight = 0;
-		int32 WriteIndex = 0;
-
-		uint8 CoefficientIndexLeft = ReadFromByteStream<uint8>(EncodedADPCMBlockLeft, ReadIndexLeft);
-
-		if (CoefficientIndexLeft >= NUM_ADAPTATION_COEFF)
-		{
-			UE_LOG(LogAudio, Error, TEXT("Decoding ADPCM block resulted in bad CoefficientIndexLeft (%d). BlockSize: %d, ReadIndexLeft: %d"), CoefficientIndexLeft, BlockSize, ReadIndexLeft);
-			return false;
-		}
-
-		ContextLeft.AdaptationDelta = ReadFromByteStream<int16>(EncodedADPCMBlockLeft, ReadIndexLeft);
-		ContextLeft.Sample1 = ReadFromByteStream<int16>(EncodedADPCMBlockLeft, ReadIndexLeft);
-		ContextLeft.Sample2 = ReadFromByteStream<int16>(EncodedADPCMBlockLeft, ReadIndexLeft);
-		ContextLeft.Coefficient1 = ContextLeft.AdaptationCoefficient1[CoefficientIndexLeft];
-		ContextLeft.Coefficient2 = ContextLeft.AdaptationCoefficient2[CoefficientIndexLeft];
-
-		uint8 CoefficientIndexRight = ReadFromByteStream<uint8>(EncodedADPCMBlockRight, ReadIndexRight);
-
-		if (CoefficientIndexRight >= NUM_ADAPTATION_COEFF)
-		{
-			UE_LOG(LogAudio, Error, TEXT("Decoding ADPCM block resulted in bad CoefficientIndexRight (%d). BlockSize: %d, ReadIndexLeft: %d"), CoefficientIndexRight, BlockSize, ReadIndexRight);
-			return false;
-		}
-
-		ContextRight.AdaptationDelta = ReadFromByteStream<int16>(EncodedADPCMBlockRight, ReadIndexRight);
-		ContextRight.Sample1 = ReadFromByteStream<int16>(EncodedADPCMBlockRight, ReadIndexRight);
-		ContextRight.Sample2 = ReadFromByteStream<int16>(EncodedADPCMBlockRight, ReadIndexRight);
-		ContextRight.Coefficient1 = ContextRight.AdaptationCoefficient1[CoefficientIndexRight];
-		ContextRight.Coefficient2 = ContextRight.AdaptationCoefficient2[CoefficientIndexRight];
-
-		// The first two samples from each stream are sent directly to the output in reverse order, as per the standard
-		DecodedPCMData[WriteIndex++] = ContextLeft.Sample2;
-		DecodedPCMData[WriteIndex++] = ContextRight.Sample2;
-		DecodedPCMData[WriteIndex++] = ContextLeft.Sample1;
-		DecodedPCMData[WriteIndex++] = ContextRight.Sample1;
-
-		uint8 EncodedNibblePairLeft = 0;
-		uint8 EncodedNibbleLeft = 0;
-		uint8 EncodedNibblePairRight = 0;
-		uint8 EncodedNibbleRight = 0;
-
-		while (ReadIndexLeft < BlockSize)
-		{
-			// Read from the byte stream and advance the read head.
-			EncodedNibblePairLeft = ReadFromByteStream<uint8>(EncodedADPCMBlockLeft, ReadIndexLeft);
-			EncodedNibblePairRight = ReadFromByteStream<uint8>(EncodedADPCMBlockRight, ReadIndexRight);
-
-			EncodedNibbleLeft = (EncodedNibblePairLeft >> 4) & 0x0F;
-			DecodedPCMData[WriteIndex++] = DecodeNibble(ContextLeft, EncodedNibbleLeft);
-
-			EncodedNibbleRight = (EncodedNibblePairRight >> 4) & 0x0F;
-			DecodedPCMData[WriteIndex++] = DecodeNibble(ContextRight, EncodedNibbleRight);
-
-			EncodedNibbleLeft = EncodedNibblePairLeft & 0x0F;
-			DecodedPCMData[WriteIndex++] = DecodeNibble(ContextLeft, EncodedNibbleLeft);
-
-			EncodedNibbleRight = EncodedNibblePairRight & 0x0F;
-			DecodedPCMData[WriteIndex++] = DecodeNibble(ContextRight, EncodedNibbleRight);
-		}
-
-		return true;
-	}
-} // end namespace ADPCM
-
-class ADPCMAUDIODECODER_API FAdpcmAudioDecoderModule : public IModuleInterface
-{
-public:	
-	TUniquePtr<IAudioInfoFactory> PcmFactory;
-	TUniquePtr<IAudioInfoFactory> AdpcmFactory;
-
-	virtual void StartupModule() override
-	{
-		AdpcmFactory = MakeUnique<FSimpleAudioInfoFactory>([] { return new FADPCMAudioInfo(); }, Audio::NAME_ADPCM);
-
-		// PCM is actually ADPCM at quality 100% 
-		PcmFactory = MakeUnique<FSimpleAudioInfoFactory>([] { return new FADPCMAudioInfo(); }, Audio::NAME_PCM);
-	}
-
-	virtual void ShutdownModule() override {}
-};
-
-IMPLEMENT_MODULE(FAdpcmAudioDecoderModule, AdpcmAudioDecoder)
\ No newline at end of file
diff --git a/Engine/Source/Runtime/AdpcmAudioDecoder/Module/Private/ADPCMAudioInfoModule.cpp b/Engine/Source/Runtime/AdpcmAudioDecoder/Module/Private/ADPCMAudioInfoModule.cpp
new file mode 100644
--- /dev/null
+++ b/Engine/Source/Runtime/AdpcmAudioDecoder/Module/Private/ADPCMAudioInfoModule.cpp
@@ -0,0 +1,59 @@
+﻿// Copyright Epic Games, Inc. All Rights Reserved.
+
+#include "AudioDecompress.h"
+#include "Decoders/ADPCMAudioInfo.h"
+#include "PcmAudioInfoHybrid.h"
+
+class ADPCMAUDIODECODER_API FAdpcmAudioDecoderModule : public IModuleInterface
+{
+public:	
+	TUniquePtr<IAudioInfoFactory> PcmFactory;
+	TUniquePtr<IAudioInfoFactory> AdpcmFactory;
+
+	template<typename T> TUniquePtr<IAudioInfoFactory> MakeFactory(const FName Name)
+	{
+		return MakeUnique<FSimpleAudioInfoFactory>([] { return new T(); }, Name);
+	}
+	void RegisterLegacy()
+	{
+		AdpcmFactory = MakeFactory<FADPCMAudioInfo>(Audio::NAME_ADPCM);
+		PcmFactory = MakeFactory<FADPCMAudioInfo>(Audio::NAME_PCM);
+	}
+	void RegisterHybrid()
+	{
+		AdpcmFactory = MakeFactory<FPcmAudioInfoHybrid>(Audio::NAME_ADPCM);
+		PcmFactory = MakeFactory<FPcmAudioInfoHybrid>(Audio::NAME_PCM);
+	}
+	void Register(const bool bUseLegacy)
+	{
+		AdpcmFactory.Reset();
+		PcmFactory.Reset();
+		if (bUseLegacy)
+		{
+			RegisterLegacy();
+		}
+		else
+		{
+			RegisterHybrid();
+		}	
+	}
+	virtual void StartupModule() override
+	{
+		static int32 bUseLegacyDecoder = 0;
+		static FAutoConsoleVariableRef CVarUseLegacyDecoder(
+		TEXT("au.adpcm.UseLegacyDecoder"),
+		bUseLegacyDecoder,
+		TEXT("0:Hybrid, 1:Legacy"),
+		 FConsoleVariableDelegate::CreateLambda([this](IConsoleVariable* Var)
+		 {
+		 	Register(Var->GetBool());
+		 }),
+		ECVF_Default);
+
+		Register(!!bUseLegacyDecoder);
+	}
+
+	virtual void ShutdownModule() override {}
+};
+
+IMPLEMENT_MODULE(FAdpcmAudioDecoderModule, AdpcmAudioDecoder)
\ No newline at end of file
diff --git a/Engine/Source/Runtime/AdpcmAudioDecoder/Module/Private/Adpcm.cpp b/Engine/Source/Runtime/AdpcmAudioDecoder/Module/Private/Adpcm.cpp
new file mode 100644
--- /dev/null
+++ b/Engine/Source/Runtime/AdpcmAudioDecoder/Module/Private/Adpcm.cpp
@@ -0,0 +1,176 @@
+﻿// Copyright Epic Games, Inc. All Rights Reserved.
+
+#include "Adpcm.h"
+#include "Decoders/ADPCMAudioInfo.h"
+
+namespace ADPCM
+{
+	using namespace ADPCMPrivate;
+	
+	template <typename T>
+	static void GetAdaptationTable(T(&OutAdaptationTable)[NUM_ADAPTATION_TABLE])
+	{
+		// Magic values as specified by standard
+		static T AdaptationTable[] =
+		{
+			230, 230, 230, 230, 307, 409, 512, 614,
+			768, 614, 512, 409, 307, 230, 230, 230
+		};
+
+		FMemory::Memcpy(&OutAdaptationTable, AdaptationTable, sizeof(AdaptationTable));
+	}
+
+	struct FAdaptationContext
+	{
+	public:
+		// Adaptation constants
+		int32 AdaptationTable[NUM_ADAPTATION_TABLE];
+		int32 AdaptationCoefficient1[NUM_ADAPTATION_COEFF];
+		int32 AdaptationCoefficient2[NUM_ADAPTATION_COEFF];
+
+		int32 AdaptationDelta;
+		int32 Coefficient1;
+		int32 Coefficient2;
+		int32 Sample1;
+		int32 Sample2;
+
+		FAdaptationContext() :
+			AdaptationDelta(0),
+			Coefficient1(0),
+			Coefficient2(0),
+			Sample1(0),
+			Sample2(0)
+		{
+			GetAdaptationTable(AdaptationTable);
+			GetAdaptationCoefficients(AdaptationCoefficient1, AdaptationCoefficient2);
+		}
+	};
+
+	FORCEINLINE int16 DecodeNibble(FAdaptationContext& Context, uint8 EncodedNibble)
+	{
+		int32 PredictedSample = (Context.Sample1 * Context.Coefficient1 + Context.Sample2 * Context.Coefficient2) / 256;
+		PredictedSample += SignExtend<int8, 4>(EncodedNibble) * Context.AdaptationDelta;
+		PredictedSample = FMath::Clamp(PredictedSample, -32768, 32767);
+
+		// Shuffle samples for the next iteration
+		Context.Sample2 = Context.Sample1;
+		Context.Sample1 = static_cast<int16>(PredictedSample);
+		Context.AdaptationDelta = (Context.AdaptationDelta * Context.AdaptationTable[EncodedNibble]) / 256;
+		Context.AdaptationDelta = FMath::Max(Context.AdaptationDelta, 16);
+
+		return Context.Sample1;
+	}
+
+	bool DecodeBlock(const uint8* EncodedADPCMBlock, int32 BlockSize, int16* DecodedPCMData)
+	{
+		FAdaptationContext Context;
+		int32 ReadIndex = 0;
+		int32 WriteIndex = 0;
+
+		uint8 CoefficientIndex = ReadFromByteStream<uint8>(EncodedADPCMBlock, ReadIndex);
+		if (CoefficientIndex >= NUM_ADAPTATION_COEFF)
+		{
+			UE_LOG(LogAudio, Error, TEXT("Decoding ADPCM block resulted in bad CoefficientIndex (%d). BlockSize: %d, ReadIndex: %d"), CoefficientIndex, BlockSize, ReadIndex);
+			return false;
+		}
+		else
+		{
+			Context.AdaptationDelta = ReadFromByteStream<int16>(EncodedADPCMBlock, ReadIndex);
+			Context.Sample1 = ReadFromByteStream<int16>(EncodedADPCMBlock, ReadIndex);
+			Context.Sample2 = ReadFromByteStream<int16>(EncodedADPCMBlock, ReadIndex);
+
+			Context.Coefficient1 = Context.AdaptationCoefficient1[CoefficientIndex];
+			Context.Coefficient2 = Context.AdaptationCoefficient2[CoefficientIndex];
+		}
+
+		// The first two samples are sent directly to the output in reverse order, as per the standard
+		DecodedPCMData[WriteIndex++] = Context.Sample2;
+		DecodedPCMData[WriteIndex++] = Context.Sample1;
+
+		uint8 EncodedNibblePair = 0;
+		uint8 EncodedNibble = 0;
+		while (ReadIndex < BlockSize)
+		{
+			// Read from the byte stream and advance the read head.
+			EncodedNibblePair = ReadFromByteStream<uint8>(EncodedADPCMBlock, ReadIndex);
+
+			EncodedNibble = (EncodedNibblePair >> 4) & 0x0F;
+			DecodedPCMData[WriteIndex++] = DecodeNibble(Context, EncodedNibble);
+
+			EncodedNibble = EncodedNibblePair & 0x0F;
+			DecodedPCMData[WriteIndex++] = DecodeNibble(Context, EncodedNibble);
+		}
+
+		return true;
+	}
+	
+	bool DecodeBlock(const uint8* EncodedADPCMBlock, const int32 BlockSize, int16* DecodedPCMData, const int32 DecodedPCMSizeSamples)
+	{
+		FAdaptationContext Context;
+		int32 ReadIndex = 0;
+		int32 WriteIndex = 0;
+
+		if (DecodedPCMSizeSamples <= 0)
+		{
+			UE_LOG(LogAudio, Error, TEXT("Decoding ADPCM block has insufficient space to decode to: DecodedPCMSize=%d"), DecodedPCMSizeSamples);
+			return false;
+		}
+		
+		const uint8 CoefficientIndex = ReadFromByteStream<uint8>(EncodedADPCMBlock, ReadIndex);
+		if (CoefficientIndex >= NUM_ADAPTATION_COEFF)
+		{
+			UE_LOG(LogAudio, Error, TEXT("Decoding ADPCM block resulted in bad CoefficientIndex (%d). BlockSize: %d, ReadIndex: %d"), CoefficientIndex, BlockSize, ReadIndex);
+			return false;
+		}
+
+		Context.AdaptationDelta = ReadFromByteStream<int16>(EncodedADPCMBlock, ReadIndex);
+		Context.Sample1 = ReadFromByteStream<int16>(EncodedADPCMBlock, ReadIndex);
+		Context.Sample2 = ReadFromByteStream<int16>(EncodedADPCMBlock, ReadIndex);
+
+		Context.Coefficient1 = Context.AdaptationCoefficient1[CoefficientIndex];
+		Context.Coefficient2 = Context.AdaptationCoefficient2[CoefficientIndex];
+
+		// The first two samples are sent directly to the output in reverse order, as per the standard
+		DecodedPCMData[WriteIndex++] = Context.Sample2;
+
+		// Handle the single frame case.
+		if (WriteIndex < DecodedPCMSizeSamples)
+		{
+			DecodedPCMData[WriteIndex++] = Context.Sample1;
+		}
+
+		// Decode full pairs.
+		uint32 RemainingSamples = DecodedPCMSizeSamples - WriteIndex;
+		uint32 FullPairs = RemainingSamples / 2;
+				
+		while (ReadIndex < BlockSize && FullPairs > 0)
+		{
+			// Read from the byte stream and advance the read head.
+			const uint8 EncodedNibblePair = ReadFromByteStream<uint8>(EncodedADPCMBlock, ReadIndex);
+
+			uint8 EncodedNibble = (EncodedNibblePair >> 4) & 0x0F;
+			DecodedPCMData[WriteIndex++] = DecodeNibble(Context, EncodedNibble);
+
+			EncodedNibble = EncodedNibblePair & 0x0F;
+			DecodedPCMData[WriteIndex++] = DecodeNibble(Context, EncodedNibble);
+
+			FullPairs--;
+			RemainingSamples -= 2;
+		}
+
+		// Residual half-pair.
+		if (RemainingSamples > 0)
+		{
+			check(RemainingSamples == 1);
+			
+			// Read from the byte stream and advance the read head.
+			const uint8 EncodedNibblePair = ReadFromByteStream<uint8>(EncodedADPCMBlock, ReadIndex);
+
+			uint8 EncodedNibble = (EncodedNibblePair >> 4) & 0x0F;
+			DecodedPCMData[WriteIndex++] = DecodeNibble(Context, EncodedNibble);	
+		}
+		
+		return true;
+	}
+} // end namespace ADPCM
+
diff --git a/Engine/Source/Runtime/AdpcmAudioDecoder/Module/Private/Adpcm.h b/Engine/Source/Runtime/AdpcmAudioDecoder/Module/Private/Adpcm.h
new file mode 100644
--- /dev/null
+++ b/Engine/Source/Runtime/AdpcmAudioDecoder/Module/Private/Adpcm.h
@@ -0,0 +1,105 @@
+﻿// Copyright Epic Games, Inc. All Rights Reserved.
+
+#pragma once
+
+#include "HAL/Platform.h"
+
+namespace ADPCM
+{
+	// Returns true if successful, false if there was an error in decoding
+	bool DecodeBlock(const uint8* EncodedADPCMBlock, const int32 BlockSize, int16* DecodedPCMData, const int32 DecodedPCMSize);
+	bool DecodeBlock(const uint8* EncodedADPCMBlock, const int32 BlockSize, int16* DecodedPCMData);
+} //namespace ADPCM
+
+namespace ADPCMPrivate
+{
+	///////////////////////////////////////////////////////`///////////////////
+	// Copied from IOS - probably want to split and share
+	//////////////////////////////////////////////////////`////////////////////
+
+	template <typename T, uint32 B>
+	inline T SignExtend(const T ValueToExtend)
+	{
+		struct { T ExtendedValue : B; } SignExtender;
+		return SignExtender.ExtendedValue = ValueToExtend;
+	}
+
+	template <typename T>
+	inline T ReadFromByteStream(const uint8* ByteStream, int32& ReadIndex, bool bLittleEndian = true)
+	{
+		T ValueRaw = 0;
+
+		if (bLittleEndian)
+		{
+#if PLATFORM_LITTLE_ENDIAN
+			for (int32 ByteIndex = 0; ByteIndex < sizeof(T); ++ByteIndex)
+#else
+			for (int32 ByteIndex = sizeof(T) - 1; ByteIndex >= 0; --ByteIndex)
+#endif // PLATFORM_LITTLE_ENDIAN
+			{
+				ValueRaw |= ByteStream[ReadIndex++] << 8 * ByteIndex;
+			}
+		}
+		else
+		{
+#if PLATFORM_LITTLE_ENDIAN
+			for (int32 ByteIndex = sizeof(T) - 1; ByteIndex >= 0; --ByteIndex)
+#else
+			for (int32 ByteIndex = 0; ByteIndex < sizeof(T); ++ByteIndex)
+#endif // PLATFORM_LITTLE_ENDIAN
+			{
+				ValueRaw |= ByteStream[ReadIndex++] << 8 * ByteIndex;
+			}
+		}
+
+		return ValueRaw;
+	}
+
+	template <typename T>
+	inline void WriteToByteStream(T Value, uint8* ByteStream, int32& WriteIndex, bool bLittleEndian = true)
+	{
+		if (bLittleEndian)
+		{
+#if PLATFORM_LITTLE_ENDIAN
+			for (int32 ByteIndex = 0; ByteIndex < sizeof(T); ++ByteIndex)
+#else
+			for (int32 ByteIndex = sizeof(T) - 1; ByteIndex >= 0; --ByteIndex)
+#endif // PLATFORM_LITTLE_ENDIAN
+			{
+				ByteStream[WriteIndex++] = (Value >> (8 * ByteIndex)) & 0xFF;
+			}
+		}
+		else
+		{
+#if PLATFORM_LITTLE_ENDIAN
+			for (int32 ByteIndex = sizeof(T) - 1; ByteIndex >= 0; --ByteIndex)
+#else
+			for (int32 ByteIndex = 0; ByteIndex < sizeof(T); ++ByteIndex)
+#endif // PLATFORM_LITTLE_ENDIAN
+			{
+				ByteStream[WriteIndex++] = (Value >> (8 * ByteIndex)) & 0xFF;
+			}
+		}
+	}
+
+	template <typename T>
+	inline T ReadFromArray(const T* ElementArray, int32& ReadIndex, int32 NumElements, int32 IndexStride = 1)
+	{
+		T OutputValue = 0;
+
+		if (ReadIndex >= 0 && ReadIndex < NumElements)
+		{
+			OutputValue = ElementArray[ReadIndex];
+			ReadIndex += IndexStride;
+		}
+
+		return OutputValue;
+	}
+
+	//////////////////////////////////////////////////////////////////////////
+	// End of copied section
+	//////////////////////////////////////////////////////////////////////////
+} // namespace ADPCMPrivate
+
+
+
diff --git a/Engine/Source/Runtime/AdpcmAudioDecoder/Module/Private/PcmAudioInfoHybrid.cpp b/Engine/Source/Runtime/AdpcmAudioDecoder/Module/Private/PcmAudioInfoHybrid.cpp
new file mode 100644
--- /dev/null
+++ b/Engine/Source/Runtime/AdpcmAudioDecoder/Module/Private/PcmAudioInfoHybrid.cpp
@@ -0,0 +1,481 @@
+// Copyright Epic Games, Inc. All Rights Reserved.
+
+#include "PcmAudioInfoHybrid.h"
+#include "Adpcm.h"
+#include "Containers/StridedView.h"
+#include "Decoders/ADPCMAudioInfo.h"
+#include "Interfaces/IAudioFormat.h"
+
+namespace AdpcmAudioInfoPrivate
+{
+	struct FAdpcmImpl final : FPcmAudioInfoHybrid::FHybridImpl
+	{
+		using FHybridImpl::FHybridImpl;
+		
+		virtual bool ParseWaveHeader(const FWaveModInfo& WaveInfo, const WaveFormatHeader* FormatHeader, const uint32 HeaderSize) override
+		{
+			if (!ensure(FormatHeader->wFormatTag == WAVE_FORMAT_ADPCM))
+			{
+				return false;
+			}
+			if (HeaderSize < sizeof(ADPCM::ADPCMFormatHeader))
+			{
+				return false;
+			}
+
+			const ADPCM::ADPCMFormatHeader* ADPCMHeader = (const ADPCM::ADPCMFormatHeader*)FormatHeader;
+			TotalSamplesPerChannel = ADPCMHeader->SamplesPerChannel;
+
+			// ADPCM starts with 2 uncompressed samples and then the remaining compressed sample data has 2 samples per byte
+			UncompressedBlockSamples = ADPCMHeader->wSamplesPerBlock;
+			UncompressedBlockSize = UncompressedBlockSamples * sizeof(int16);
+			CompressedBlockSize = *WaveInfo.pBlockAlign;
+
+			// Ensure TotalDecodedSize is a even multiple of the compressed block size so that the buffer is not over read on the last block
+			TotalBlocks = ((WaveInfo.SampleDataSize + CompressedBlockSize - 1) / CompressedBlockSize);
+			TotalBlocksPerChannel = TotalBlocks / NumChannels;
+			TotalDecodedSize = TotalBlocks * UncompressedBlockSize;
+
+			Reservoir.SetNum(NumChannels);
+
+			// Encoded format explanation:
+			// 
+			// The ordering is different if we're streaming or not (for legacy reasons).
+			// Streaming format is interleaved blocks.
+			// Non-streaming (i.e. ForceInline) is multi-mono blocks.
+			// [L][R][C][LFE][LS][RS][L][R][C][LFE][LS][RS] - Streaming format. (blocks are interleaved by channel).
+			// [L][L][R][R][C][C][LFE][LFE][LS][LS][RS][RS] - Non streaming. (channels are multi mono).
+			// 
+			// This affects how we traverse the source data in the decoder. But is abstracted away by the
+			// following two stride values:
+			//
+			// ChannelBlockStride - This is how much in bytes to jump over to get to the next channels data.
+			// SrcAdvanceStride - This is how much to advance the source pointer by after reading an entire frame-block (i.e. blocksize * nchannels)
+
+			if (Owner->IsStreaming())
+			{
+				ChannelBlockStride = CompressedBlockSize;							// 1 block 
+				SrcAdvanceStride = CompressedBlockSize * NumChannels;				// 1 block "frame"
+			}
+			else
+			{
+				ChannelBlockStride = CompressedBlockSize * TotalBlocksPerChannel;	// 1 entire channels full set of blocks.
+				SrcAdvanceStride = CompressedBlockSize;								// 1 block.
+			}
+
+			// success.
+			return true;
+		}
+
+		virtual int32 GetFrameSize() override
+		{
+			return CompressedBlockSize * NumChannels;
+		}
+
+		virtual uint32 GetMaxFrameSizeSamples() const override
+		{
+			return UncompressedBlockSamples; 
+		}
+
+		virtual void SeekToFrame(const uint32 InSeekFrame) override
+		{
+			using namespace AdpcmAudioInfoPrivate;
+			const uint32 SeekFrameClamped = FMath::Clamp<uint32>(InSeekFrame, 0, TotalSamplesPerChannel - 1);
+
+			// Which block is frame in? 
+			const uint32 SeekFrameBlockIndex = SeekFrameClamped / UncompressedBlockSamples;
+			check(SeekFrameBlockIndex < TotalBlocksPerChannel);
+
+			// Abs offset of first block in AudioData
+			const uint32 AbsBlockOffset = SrcAdvanceStride * SeekFrameBlockIndex;
+
+			// Offset in frames of start of block.
+			const uint32 BlockFrameStart = SeekFrameBlockIndex * UncompressedBlockSamples;
+			check(BlockFrameStart < TotalSamplesPerChannel);
+
+			// Residual samples we need to skip.
+			const int32 FramesToSkip = SeekFrameClamped - BlockFrameStart;
+			check(FramesToSkip >= 0);
+			check(FramesToSkip < (int32)UncompressedBlockSamples);
+
+			// Do the seek.
+			TotalFramesDecoded = BlockFrameStart;
+			NumFramesToSkip = FramesToSkip;
+			Owner->SeekToAbs(AbsBlockOffset, InSeekFrame);
+
+			// Clear reservoir as we've just invalidated it.
+			ReservoirStart = 0;
+			for (int32 ChannelIndex = 0; ChannelIndex < NumChannels; ++ChannelIndex)
+			{
+				Reservoir[ChannelIndex].SetNum(0, EAllowShrinking::No);
+			}
+		}
+
+		virtual void PrepareToLoop() override
+		{
+			// This is called just before we loop (if the SoundWave is looping).
+			// Need to reset or our internal accounting of the final block will be wrong.
+			TotalFramesDecoded = 0;
+		}
+
+		virtual FDecodeResult Decode(const uint8* CompressedData, const int32 CompressedDataSize, uint8* OutPCMData, const int32 OutputPCMDataSize) override
+		{
+			uint32 DstSizeFrames = OutputPCMDataSize / sizeof(int16) / NumChannels;;
+			int16* Dst = (int16*)OutPCMData;
+			uint8 const* Src = (uint8*)CompressedData;
+			uint32 SrcSize = CompressedDataSize;
+
+			while (true)
+			{
+				// Drain reservoir (interleaved copy out)
+				if (!DrainReservoir(Dst, DstSizeFrames))
+				{
+					// Copied out all we can and Dst is full.
+					break;
+				}
+
+				// Produce full buffer except for the last block which will be limited to how many frames are needed.
+				const int32 NumFramesToProduce = FMath::Min(UncompressedBlockSamples, TotalSamplesPerChannel - TotalFramesDecoded);
+				if (NumFramesToProduce == 0)
+				{
+					break;
+				}
+
+				// Decode next block of frames into reservoir.
+				if (bool bHasError = false; !DecodeNextFrameBlock(Src, SrcSize, NumFramesToProduce, bHasError))
+				{
+					// Decode either ran of source, or had an error.
+					if (bHasError)
+					{
+						Owner->SetErrorFlag();
+					}
+					break;
+				}
+
+				TotalFramesDecoded += NumFramesToProduce;
+			}
+
+			FDecodeResult Result;
+			Result.NumPcmBytesProduced = UE_PTRDIFF_TO_INT32((uint8*)Dst - OutPCMData);
+			Result.NumAudioFramesProduced = Result.NumPcmBytesProduced / NumChannels / sizeof(uint16);
+			Result.NumCompressedBytesConsumed = UE_PTRDIFF_TO_INT32((uint8*)Src - CompressedData);
+			return Result;
+		}	
+
+		bool DrainReservoir(int16*& OutDst, uint32& OutDstSizeFrames)
+		{
+			check(NumChannels > 0);
+			check(NumChannels == Reservoir.Num());
+
+			// Handle skipping.
+			if (NumFramesToSkip > 0)
+			{
+				const int32 NumFramesBeforeSkip = Reservoir[0].Num() - ReservoirStart;
+				check(NumFramesBeforeSkip >= 0);
+				const int32 NumFramesToSkipInBuffer = FMath::Min<int32>(NumFramesBeforeSkip, NumFramesToSkip);
+				ReservoirStart += NumFramesToSkipInBuffer;
+				NumFramesToSkip -= NumFramesToSkipInBuffer;
+			}
+
+			const int32 NumSrcFrames = Reservoir[0].Num() - ReservoirStart;
+			check(NumSrcFrames >= 0);
+			const int32 NumDstFrames = OutDstSizeFrames;
+			const int32 NumFramesToCopy = FMath::Min(NumDstFrames, NumSrcFrames);
+
+			// Anything remaining?
+			if (NumFramesToCopy > 0)
+			{
+				switch (NumChannels)
+				{
+					case 1: // No need for interleave.
+					{
+						FMemory::Memcpy(OutDst, Reservoir[0].GetData() + ReservoirStart, NumFramesToCopy * sizeof(int16));
+						break;
+					}
+					default:
+					{
+						// We don't appear to have fast int16 interleave yet, so just hand roll one for now.
+						for (int32 Channel = 0; Channel < NumChannels; ++Channel)
+						{
+							TStridedView<int16> DstChannel = MakeStridedView(sizeof(int16) * NumChannels, OutDst + Channel, NumFramesToCopy);
+							const int16* Src = Reservoir[Channel].GetData() + ReservoirStart;
+							for (int32 Frame = 0; Frame < NumFramesToCopy; ++Frame)
+							{
+								DstChannel[Frame] = *Src++;
+							}
+						}
+						break;
+					}
+				}
+				
+				// Empty Reservoir
+				ReservoirStart += NumFramesToCopy;;
+				if (ReservoirStart == Reservoir[0].Num())
+				{
+					for (int32 Channel = 0; Channel < NumChannels; ++Channel)
+					{
+						Reservoir[Channel].SetNum(0, EAllowShrinking::No);
+					}
+					ReservoirStart = 0;
+				}
+				
+				// Subtract what we just wrote from remaining.
+				OutDstSizeFrames -= NumFramesToCopy;
+
+				// Advance ptr.
+				OutDst += NumFramesToCopy * NumChannels;
+			}
+
+			// Return true if there's still more room to copy.
+			return OutDstSizeFrames > 0;
+		}
+
+		bool DecodeNextFrameBlock(const uint8*& Src, uint32& SrcSize, const int32 NumFramesToProduce, bool& bHasError)
+		{
+			// Enough Compressed Src to decode all channels? All Src data is multiple of a block size.
+			// The final block is still block length, we just decode less of it.
+			if ((int32)SrcSize < CompressedBlockSize * NumChannels)
+			{
+				// Need more data
+				return false;
+			}
+
+			// Decode next block per channel. (~512 samples, or less. For last block we give a smaller buffer to only produce what we need.)
+			for (int32 Channel = 0; Channel < NumChannels; ++Channel)
+			{
+				const uint8* SrcChannelBlock = Src + (ChannelBlockStride * Channel);
+
+				// Make space.
+				TArray<int16>& DstChannel = Reservoir[Channel];
+				const int32 NumBefore = DstChannel.Num() - ReservoirStart;
+				check(NumBefore >= 0);
+				DstChannel.SetNum(NumBefore + NumFramesToProduce, EAllowShrinking::No);
+
+				if (!ADPCM::DecodeBlock(SrcChannelBlock, CompressedBlockSize, DstChannel.GetData() + NumBefore, NumFramesToProduce))
+				{
+					bHasError = true;
+					return false;
+				}
+			}
+
+			// Advance src.
+			SrcSize -= SrcAdvanceStride;
+			Src += SrcAdvanceStride;
+
+			// We made some data successfully.
+			return true;
+		}
+		
+		uint32 UncompressedBlockSamples = 0;	// Size of a block uncompressed in samples.
+		uint32 UncompressedBlockSize = 0;		// Size of a block uncompressed in bytes.
+		uint32 CompressedBlockSize = 0;			// Compressed block size in bytes.
+		uint32 ChannelBlockStride = 0;			// How much to step (in bytes) over the source data to get to the next channel
+		uint32 SrcAdvanceStride = 0;			// How much to advance the source by after consuming a frame block (all channels) worth of data.
+		uint32 NumFramesToSkip = 0;				// Number of frames to skip in the reservoir. Used for seeking when we need to discard some of a block
+		uint32 TotalBlocks = 0;					// Total number of blocks in the source. Used for sanity checking state. 
+		uint32 TotalBlocksPerChannel = 0;		// Total number of blocks per channel in the source. 
+		TArray<TArray<int16>> Reservoir;		// Multi-mono buffer, sized by NumChannels, that's the output of each channels block decode.
+		uint32 ReservoirStart = 0;				// The start of valid data in each of the reservoir mono buffers. (so we don't need to memmove)
+		uint32 TotalFramesDecoded = 0;			// The number of samples decoded so far (per channel)	
+	};
+	
+	struct FPcmImpl final : FPcmAudioInfoHybrid::FHybridImpl
+	{
+		using FHybridImpl::FHybridImpl;
+
+		virtual bool ParseWaveHeader(const FWaveModInfo& WaveInfo, const WaveFormatHeader* FormatHeader, const uint32 HeaderSize) override
+		{
+			if (!ensure(FormatHeader->wFormatTag == WAVE_FORMAT_LPCM))
+			{
+				return false;
+			}
+	
+			// There are no "blocks" in this case
+			TotalDecodedSize = WaveInfo.SampleDataSize;
+			TotalSamplesPerChannel = TotalDecodedSize / sizeof(uint16) / NumChannels;
+
+			return true;
+		}
+
+		virtual FDecodeResult Decode(const uint8* CompressedData, const int32 CompressedDataSize, uint8* OutPCMData, const int32 OutputPCMDataSize) override
+		{
+			const int32 ConsumedBytes = FMath::Min(OutputPCMDataSize, CompressedDataSize);
+
+			// Copy PCM directly out.
+			FMemory::Memcpy(OutPCMData, CompressedData, ConsumedBytes);
+
+			FDecodeResult Result;
+			Result.NumAudioFramesProduced = (ConsumedBytes / sizeof(int16)) / NumChannels;
+			Result.NumCompressedBytesConsumed = ConsumedBytes;
+			Result.NumPcmBytesProduced = ConsumedBytes;
+			return Result;
+		}
+
+		virtual int32 GetFrameSize() override
+		{
+			// No framing on PCM, so use the buffer size assuming we can provide that many.
+			const uint32 SrcBufferDataSize = Owner->GetSrcBufferDataSize();
+			const uint32 SrcBufferOffset = Owner->GetSrcBufferOffset();
+			const int32 RemainingInChunk = FMath::Max<int32>(SrcBufferDataSize - SrcBufferOffset, 0);
+			const int32 PcmStreamingFrame = FMath::Min<int32>(RemainingInChunk, MONO_PCM_BUFFER_SIZE);
+
+			return PcmStreamingFrame;
+		}
+
+		virtual uint32 GetMaxFrameSizeSamples() const override
+		{
+			return MONO_PCM_BUFFER_SAMPLES;
+		}
+
+		virtual void SeekToFrame(const uint32 InSeekFrame) override
+		{
+			const uint32 SeekFrameClamped = FMath::Clamp<uint32>(InSeekFrame, 0, TotalSamplesPerChannel - 1);
+			const uint32 AbsOffset = sizeof(int16) * NumChannels * SeekFrameClamped;
+			Owner->SeekToAbs(AbsOffset, SeekFrameClamped);
+		}
+	};
+}// namespace AdpcmAudioInfoPrivate
+
+bool FPcmAudioInfoHybrid::AbsPositionToChunkIndexAndOffset(const uint64 InAbsPosition, uint32& OutChunkIndex, uint32& OutChunkOffset) const
+{
+	// O(n) search through the possible chunks. This number is pretty small typically.
+	// But we could make this faster by just special casing zero, first, rest, and last.
+	
+	// If we're streaming, grab the streaming wave instance.
+	if (const FSoundWaveProxyPtr& Wave = GetStreamingSoundWave(); Wave.IsValid() && Wave->IsStreaming())
+	{
+		const int32 NumChunks = Wave->GetNumChunks();
+		uint64 AbsChunkOffset = AudioDataOffset;
+		for (int32 ChunkIndex = AudioDataChunkIndex; ChunkIndex < NumChunks; ++ChunkIndex)
+		{
+			const FStreamedAudioChunk& Chunk = Wave->GetChunk(ChunkIndex);
+			if (InAbsPosition >= AbsChunkOffset && InAbsPosition < AbsChunkOffset + Chunk.AudioDataSize)
+			{
+				OutChunkIndex = ChunkIndex;
+				OutChunkOffset = IntCastChecked<uint32>(InAbsPosition - AbsChunkOffset);
+				return true;
+			}
+			AbsChunkOffset += Chunk.AudioDataSize;
+		}
+	}
+	return false;
+}
+
+bool FPcmAudioInfoHybrid::ParseHeader(const uint8* InSrcBufferData, uint32 InSrcBufferDataSize, FSoundQualityInfo* QualityInfo)
+{
+	SrcBufferData = InSrcBufferData;
+	SrcBufferDataSize = InSrcBufferDataSize;
+	SrcBufferOffset = 0;
+	CurrentSampleCount = 0;
+
+	// We only have the header if we're streaming, so just parse that.
+	const bool bParseJustHeader = GetStreamingSoundWave().IsValid() && GetStreamingSoundWave()->IsStreaming();
+
+	void* FormatHeader = nullptr;
+	FWaveModInfo WaveInfo;
+	if (!WaveInfo.ReadWaveInfo(InSrcBufferData, InSrcBufferDataSize, nullptr, bParseJustHeader, &FormatHeader))
+	{
+		UE_LOG(LogAudio, Error, TEXT("%hs WaveInfo.ReadWaveInfo Failed"), __func__);
+		return false;
+	}
+	if (!FormatHeader)
+	{
+		UE_LOG(LogAudio, Error, TEXT("%hs WaveInfo.ReadInfo Failed to return a Header"), __func__);
+		return false;
+	}
+
+	// Make sure header is in bounds.
+	const uint32 HeaderSize = UE_PTRDIFF_TO_UINT32(WaveInfo.SampleDataStart - InSrcBufferData);
+	if (HeaderSize < sizeof(WaveFormatHeader))
+	{
+		UE_LOG(LogAudio, Error, TEXT("%hs HeaderSize=%u"), __func__, HeaderSize);
+		return false;	
+	}
+
+	// Sanity check everything looks sane.
+	NumChannels = *WaveInfo.pChannels;
+	if (NumChannels <= 0)
+	{
+		UE_LOG(LogAudio, Error, TEXT("%hs NumChannels=%u"), __func__, NumChannels);
+		return false;
+	}
+	const uint32 SampleRate = *WaveInfo.pSamplesPerSec;
+	if (SampleRate == 0)
+	{
+		UE_LOG(LogAudio, Error, TEXT("%hs SampleRate=%u"), __func__, SampleRate);
+		return false;	
+	}
+
+	// Create the appropriate impl
+	Impl = CreateImpl(*WaveInfo.pFormatTag, NumChannels, SampleRate);
+	if (!Impl.IsValid())
+	{
+		UE_LOG(LogAudio, Error, TEXT("%hs, Failed to Create Impl for %u"), __func__, *WaveInfo.pFormatTag);
+		return false;
+	}
+	
+	// Do further format specific stuff.
+	if (!Impl->ParseWaveHeader(WaveInfo, static_cast<const WaveFormatHeader*>(FormatHeader), HeaderSize))
+	{
+		return false;
+	}
+	
+	// Put the read cursor after the header.
+	SrcBufferOffset += HeaderSize;
+	AudioDataOffset = SrcBufferOffset;
+	AudioDataChunkIndex = CurrentChunkIndex; 
+	
+	// Record the total number of samples.
+	TrueSampleCount = Impl->GetTotalSamplesPerChannel() * NumChannels; 
+	
+	if (QualityInfo)
+	{
+		QualityInfo->SampleRate = *WaveInfo.pSamplesPerSec;
+		QualityInfo->NumChannels = *WaveInfo.pChannels;
+		QualityInfo->SampleDataSize = Impl->GetTotalSamplesPerChannel() * NumChannels * sizeof(int16); // Impl->GetTotalDecodedSize();
+		QualityInfo->Duration = (float)Impl->GetTotalSamplesPerChannel() / QualityInfo->SampleRate;
+	}
+	return true;
+}
+
+TUniquePtr<FPcmAudioInfoHybrid::FHybridImpl> FPcmAudioInfoHybrid::CreateImpl(const uint8 InFormatId, const int32 InNumChannels, const uint32 InSampleRate)
+{
+	using namespace AdpcmAudioInfoPrivate;
+	switch (InFormatId)
+	{
+	case WAVE_FORMAT_LPCM:
+		return MakeUnique<FPcmImpl>(this, InNumChannels, InSampleRate);
+	case WAVE_FORMAT_ADPCM:
+		return MakeUnique<FAdpcmImpl>(this,  InNumChannels, InSampleRate );
+	default:
+		break;
+	}
+	return {};
+}
+
+void FPcmAudioInfoHybrid::SeekToAbs(const uint64 InAbsPosition, const uint64 InSeekFrame)
+{
+	if (IsStreaming())
+	{
+		if (uint32 ChunkIndex, ChunkOffset; AbsPositionToChunkIndexAndOffset(InAbsPosition, ChunkIndex, ChunkOffset))
+		{
+			SetSeekBlockIndex(ChunkIndex);
+			SetSeekBlockOffset(ChunkOffset);
+			SetCurrentSampleCount(InSeekFrame);
+		}
+	}
+	else // Not streaming.
+	{
+		SetSrcBufferOffset(GetAudioDataOffset() + InAbsPosition);
+		SetCurrentSampleCount(InSeekFrame);
+	}	
+}
+
+void FPcmAudioInfoHybrid::FHybridImpl::SeekToTime(const float InSeekTime)
+{
+	const uint32 SeekFrame = InSeekTime * SampleRate;
+	const uint32 SeekFrameClamped = FMath::Clamp<uint32>( SeekFrame, 0, TotalSamplesPerChannel - 1);
+	SeekToFrame(SeekFrameClamped);
+}
+
+
diff --git a/Engine/Source/Runtime/AdpcmAudioDecoder/Module/Private/PcmAudioInfoHybrid.h b/Engine/Source/Runtime/AdpcmAudioDecoder/Module/Private/PcmAudioInfoHybrid.h
new file mode 100644
--- /dev/null
+++ b/Engine/Source/Runtime/AdpcmAudioDecoder/Module/Private/PcmAudioInfoHybrid.h
@@ -0,0 +1,113 @@
+﻿// Copyright Epic Games, Inc. All Rights Reserved.
+
+#pragma once
+
+#include "AudioDecompress.h"
+#include "Decoders/ADPCMAudioInfo.h"
+
+class FPcmAudioInfoHybrid : public IStreamedCompressedInfo
+{
+public:
+	FPcmAudioInfoHybrid() = default;
+	virtual ~FPcmAudioInfoHybrid() override = default;
+
+	class FHybridImpl
+	{
+	public:
+		UE_NONCOPYABLE(FHybridImpl)
+		
+		FHybridImpl(FPcmAudioInfoHybrid* InOwner, const int32 InNumChannels, const uint32 InSampleRate)
+			: Owner(InOwner)
+			, NumChannels(InNumChannels)
+			, SampleRate(InSampleRate)
+		{}
+						
+		virtual ~FHybridImpl() = default;
+		virtual void PrepareToLoop() {}
+		virtual bool ParseWaveHeader(const FWaveModInfo& WaveInfo, const WaveFormatHeader* FormatHeader, const uint32 HeaderSize) = 0;
+		virtual int32 GetFrameSize() = 0;
+		virtual uint32 GetMaxFrameSizeSamples() const = 0;
+		virtual void SeekToFrame(const uint32 InSeekFrame) = 0;
+		virtual void SeekToTime(const float InSeekTime);
+		virtual FDecodeResult Decode(const uint8* CompressedData, const int32 CompressedDataSize, uint8* OutPCMData, const int32 OutputPCMDataSize) = 0;
+		
+		uint32 GetTotalDecodedSize() const { return TotalDecodedSize; }
+		uint32 GetTotalSamplesPerChannel() const { return TotalSamplesPerChannel; }
+		
+	protected:
+		FPcmAudioInfoHybrid* Owner = nullptr;
+		int32 NumChannels = 0;
+		uint32 TotalDecodedSize = 0;
+		uint32 SampleRate = 0;
+		uint32 TotalSamplesPerChannel = 0;			// Number of samples per channel, used to detect when an audio waveform has ended
+	};
+	
+protected:
+
+	/** Parse the header information from the input source buffer data. This is dependent on compression format. */
+	virtual bool ParseHeader(const uint8* InSrcBufferData, uint32 InSrcBufferDataSize, FSoundQualityInfo* QualityInfo) override;
+
+	/** Create the compression format dependent decoder object. */
+	virtual bool CreateDecoder() override { return true; }
+
+	virtual void PrepareToLoop() override
+	{
+		Impl->PrepareToLoop();
+	}
+	virtual int32 GetFrameSize() override
+	{
+		return Impl->GetFrameSize();
+	}
+	virtual uint32 GetMaxFrameSizeSamples() const override
+	{
+		return Impl->GetMaxFrameSizeSamples();
+	}
+	virtual void SeekToFrame(const uint32 InSeekFrame) override
+	{
+		return Impl->SeekToFrame(InSeekFrame);
+	}
+	virtual void SeekToTime(const float InSeekTime) override
+	{
+		return Impl->SeekToTime(InSeekTime);
+	}
+	virtual FDecodeResult Decode(const uint8* CompressedData, const int32 CompressedDataSize, uint8* OutPCMData, const int32 OutputPCMDataSize) override
+	{
+		return Impl->Decode(CompressedData, CompressedDataSize, OutPCMData, OutputPCMDataSize);	
+	}
+
+public:
+	bool IsStreaming() const
+	{
+		return GetStreamingSoundWave().IsValid() && GetStreamingSoundWave()->IsStreaming();
+	}
+	void SeekToAbs(const uint64 InAbsPosition, const uint64 InSeekFrame);	
+	void SetErrorFlag() const { bHasError = true; }
+	
+	uint32 GetAudioDataOffset() const { return AudioDataOffset; }
+	uint32 GetSrcBufferOffset() const { return SrcBufferOffset; }
+	uint32 GetSrcBufferDataSize() const { return SrcBufferDataSize; }
+	
+private:
+	TUniquePtr<FHybridImpl> CreateImpl(const uint8 InFormatId, const int32 InNumChannels, const uint32 InSampleRate);
+	bool AbsPositionToChunkIndexAndOffset(const uint64 AbsPosition, uint32& ChunkIndex, uint32& ChunkOffset) const;
+	
+	void SetSeekBlockIndex(const uint32 InStreamSeekBlockIndex)
+	{
+		ensure(GetStreamingSoundWave().IsValid() && GetStreamingSoundWave()->IsStreaming()
+			&& InStreamSeekBlockIndex < GetStreamingSoundWave()->GetNumChunks());
+		StreamSeekBlockIndex = InStreamSeekBlockIndex;
+	};
+	void SetSeekBlockOffset(const int32 InStreamSeekBlockOffset)
+	{
+		StreamSeekBlockOffset = InStreamSeekBlockOffset;
+	};
+	void SetCurrentSampleCount(const uint32 InSampleCount) { CurrentSampleCount = InSampleCount; }
+	void SetSrcBufferOffset(const uint32 InSrcBufferOffset)
+	{
+		ensure(InSrcBufferOffset < SrcBufferDataSize); 
+		SrcBufferOffset = InSrcBufferOffset;
+	}
+
+	friend class FHybridImpl;
+	TUniquePtr<FHybridImpl> Impl;
+};
diff --git a/Engine/Source/Runtime/Engine/Private/AudioDecompress.cpp b/Engine/Source/Runtime/Engine/Private/AudioDecompress.cpp
--- a/Engine/Source/Runtime/Engine/Private/AudioDecompress.cpp
+++ b/Engine/Source/Runtime/Engine/Private/AudioDecompress.cpp
@@ -348,6 +348,9 @@ bool IStreamedCompressedInfo::StreamCompressedData(uint8* Destination, bool bLoo
 			CurrentChunkIndex = StreamSeekBlockIndex;
 			SrcBufferDataSize = ChunkSize;
 			SrcBufferOffset = StreamSeekBlockOffset;
+			LastPCMByteSize = 0;
+			LastPCMOffset = 0;
+			bStoringEndOfFile = false;
 		}
 		else
 		{
diff --git a/Engine/Source/Runtime/Engine/Public/AudioDecompress.h b/Engine/Source/Runtime/Engine/Public/AudioDecompress.h
--- a/Engine/Source/Runtime/Engine/Public/AudioDecompress.h
+++ b/Engine/Source/Runtime/Engine/Public/AudioDecompress.h
@@ -165,7 +165,7 @@ public:
 	virtual int32 GetCurrentChunkOffset() const {return -1;}
 
 	/** Return the streaming sound wave used by this decoder. Returns nullptr if there is not a streaming sound wave. */
-	virtual const FSoundWaveProxyPtr& GetStreamingSoundWave() { return StreamingSoundWave; }
+	virtual const FSoundWaveProxyPtr& GetStreamingSoundWave() const { return StreamingSoundWave; }
 
 protected:
 	mutable bool bHasError = false;
